GeometryProcessing: add new prototype GeometryFlow plugin for implementing expensive geometry processing operations in a node-graph style. Add MeshLODToolset plugin that implements a GeometryFlow graph for processing a high-resolution source mesh+materials into a game-ready staticmesh asset + material instances + collision geometry. GenerateStaticMeshLODAssetTool exposes this process in a Modeling Mode tool.
#rnx

diff --git a/Engine/Plugins/Experimental/GeometryFlow/GeometryFlow.uplugin b/Engine/Plugins/Experimental/GeometryFlow/GeometryFlow.uplugin
new file mode 100644
--- /dev/null
+++ b/Engine/Plugins/Experimental/GeometryFlow/GeometryFlow.uplugin
@@ -0,0 +1,45 @@
+{
+	"FileVersion": 3,
+	"Version": 1,
+	"VersionName": "0.1",
+	"FriendlyName": "GeometryFlow",
+	"Description": "Geometry DataFlow Graph",
+	"Category": "Other",
+	"CreatedBy": "Epic Games, Inc.",
+	"CreatedByURL": "http://epicgames.com",
+	"DocsURL": "",
+	"MarketplaceURL": "",
+	"SupportURL": "",
+	"CanContainContent": true,
+	"IsBetaVersion": true,
+	"Installed": false,
+	"Hidden": true,
+	"Modules": [
+		{
+			"Name": "GeometryFlowCore",
+			"Type": "Runtime",
+			"LoadingPhase": "Default"
+		},
+		{
+			"Name": "GeometryFlowMeshProcessing",
+			"Type": "Runtime",
+			"LoadingPhase": "Default"
+		},		
+		{
+			"Name": "GeometryFlowMeshProcessingEditor",
+			"Type": "Editor",
+			"LoadingPhase": "Default",
+			"WhitelistPlatforms": [ "Win64" ]
+		},		
+	],
+	"Plugins": [
+		{
+			"Name": "GeometryProcessing",
+			"Enabled": true
+		},
+		{
+			"Name": "MeshModelingToolset",
+			"Enabled": true
+		}
+	]
+}
diff --git a/Engine/Plugins/Experimental/GeometryFlow/Source/GeometryFlowCore/GeometryFlowCore.Build.cs b/Engine/Plugins/Experimental/GeometryFlow/Source/GeometryFlowCore/GeometryFlowCore.Build.cs
new file mode 100644
--- /dev/null
+++ b/Engine/Plugins/Experimental/GeometryFlow/Source/GeometryFlowCore/GeometryFlowCore.Build.cs
@@ -0,0 +1,50 @@
+// Copyright Epic Games, Inc. All Rights Reserved.
+
+using UnrealBuildTool;
+
+public class GeometryFlowCore : ModuleRules
+{
+	public GeometryFlowCore(ReadOnlyTargetRules Target) : base(Target)
+	{
+		PCHUsage = ModuleRules.PCHUsageMode.UseExplicitOrSharedPCHs;
+		
+		PublicIncludePaths.AddRange(
+			new string[] {
+				// ... add public include paths required here ...
+			}
+			);
+				
+		
+		PrivateIncludePaths.AddRange(
+			new string[] {
+				// ... add other private include paths required here ...
+			}
+			);
+			
+		
+		PublicDependencyModuleNames.AddRange(
+			new string[]
+			{
+				"Core",
+				"GeometricObjects"
+			}
+            );
+			
+		
+		PrivateDependencyModuleNames.AddRange(
+			new string[]
+			{
+				//"CoreUObject",
+				//"Engine"
+			}
+			);
+		
+		
+		DynamicallyLoadedModuleNames.AddRange(
+			new string[]
+			{
+				// ... add any modules that your module loads dynamically here ...
+			}
+			);
+	}
+}
diff --git a/Engine/Plugins/Experimental/GeometryFlow/Source/GeometryFlowCore/Private/GeometryFlowCoreModule.cpp b/Engine/Plugins/Experimental/GeometryFlow/Source/GeometryFlowCore/Private/GeometryFlowCoreModule.cpp
new file mode 100644
--- /dev/null
+++ b/Engine/Plugins/Experimental/GeometryFlow/Source/GeometryFlowCore/Private/GeometryFlowCoreModule.cpp
@@ -0,0 +1,20 @@
+// Copyright Epic Games, Inc. All Rights Reserved.
+
+#include "GeometryFlowCoreModule.h"
+
+#define LOCTEXT_NAMESPACE "FGeometryFlowCoreModule"
+
+void FGeometryFlowCoreModule::StartupModule()
+{
+	// This code will execute after your module is loaded into memory; the exact timing is specified in the .uplugin file per-module
+}
+
+void FGeometryFlowCoreModule::ShutdownModule()
+{
+	// This function may be called during shutdown to clean up your module.  For modules that support dynamic reloading,
+	// we call this function before unloading the module.
+}
+
+#undef LOCTEXT_NAMESPACE
+	
+IMPLEMENT_MODULE(FGeometryFlowCoreModule, GeometryFlowCore)
diff --git a/Engine/Plugins/Experimental/GeometryFlow/Source/GeometryFlowCore/Private/GeometryFlowGraph.cpp b/Engine/Plugins/Experimental/GeometryFlow/Source/GeometryFlowCore/Private/GeometryFlowGraph.cpp
new file mode 100644
--- /dev/null
+++ b/Engine/Plugins/Experimental/GeometryFlow/Source/GeometryFlowCore/Private/GeometryFlowGraph.cpp
@@ -0,0 +1,244 @@
+// Copyright Epic Games, Inc. All Rights Reserved.
+
+#include "GeometryFlowGraph.h"
+
+using namespace UE::GeometryFlow;
+
+
+TSafeSharedPtr<FNode> FGraph::FindNode(FHandle NodeHandle) const
+{
+	const FNodeInfo* Found = AllNodes.Find(NodeHandle.Identifier);
+	if (!ensure(Found != nullptr))
+	{
+		return nullptr;
+	}
+	return Found->Node;
+}
+
+EGeometryFlowResult FGraph::GetInputTypeForNode(FHandle NodeHandle, FString InputName, int32& Type) const
+{
+	const FNodeInfo* Found = AllNodes.Find(NodeHandle.Identifier);
+	if (!ensure(Found != nullptr))
+	{
+		return EGeometryFlowResult::NodeDoesNotExist;
+	}
+	EGeometryFlowResult Result = Found->Node->GetInputType(InputName, Type);
+	return Result;
+}
+
+EGeometryFlowResult FGraph::GetOutputTypeForNode(FHandle NodeHandle, FString OutputName, int32& Type) const
+{
+	const FNodeInfo* Found = AllNodes.Find(NodeHandle.Identifier);
+	if (!ensure(Found != nullptr))
+	{
+		return EGeometryFlowResult::NodeDoesNotExist;
+	}
+	EGeometryFlowResult Result = Found->Node->GetOutputType(OutputName, Type);
+	return Result;
+}
+
+
+ENodeCachingStrategy FGraph::GetCachingStrategyForNode(FHandle NodeHandle) const
+{
+	const FNodeInfo* Found = AllNodes.Find(NodeHandle.Identifier);
+	if (!ensure(Found != nullptr))
+	{
+		return ENodeCachingStrategy::AlwaysCache;
+	}
+	return (Found->CachingStrategy == ENodeCachingStrategy::Default) ?
+		DefaultCachingStrategy : Found->CachingStrategy;
+}
+
+
+EGeometryFlowResult FGraph::AddConnection(FHandle FromNode, FString FromOutput, FHandle ToNode, FString ToInput)
+{
+	int32 FromType = -1, ToType = -2;
+	EGeometryFlowResult FromTypeResult = GetOutputTypeForNode(FromNode, FromOutput, FromType);
+	if (!ensure(FromTypeResult == EGeometryFlowResult::Ok))
+	{
+		return FromTypeResult;
+	}
+	EGeometryFlowResult ToTypeResult = GetInputTypeForNode(ToNode, ToInput, ToType);
+	if (!ensure(ToTypeResult == EGeometryFlowResult::Ok))
+	{
+		return ToTypeResult;
+	}
+
+	if (!ensure(FromType == ToType))
+	{
+		return EGeometryFlowResult::UnmatchedTypes;
+	}
+
+	// TODO: check for unique inputs
+
+	Connections.Add({ FromNode, FromOutput, ToNode, ToInput });
+
+	return EGeometryFlowResult::Ok;
+}
+
+
+
+EGeometryFlowResult FGraph::InferConnection(FHandle FromNodeHandle, FHandle ToNodeHandle)
+{
+	TSafeSharedPtr<FNode> FromNode = FindNode(FromNodeHandle);
+	TSafeSharedPtr<FNode> ToNode = FindNode(ToNodeHandle);
+	if (FromNode == nullptr || ToNode == nullptr)
+	{
+		ensure(false);
+		return EGeometryFlowResult::NodeDoesNotExist;
+	}
+
+	FString FromOutputName;
+	FString ToInputName;
+	int32 TotalMatchesFound = 0;
+	FromNode->EnumerateOutputs([&](const FString& OutputName, const TUniquePtr<INodeOutput>& Output)
+	{
+		int32 OutputType = Output->GetDataType();
+		ToNode->EnumerateInputs([&](const FString& InputName, const TUniquePtr<INodeInput>& Input)
+		{
+			int32 InputType = Input->GetDataType();
+			if (OutputType == InputType)
+			{
+				TotalMatchesFound++;
+				FromOutputName = OutputName;
+				ToInputName = InputName;
+			}
+		});
+	});
+	ensure(TotalMatchesFound == 1);
+	switch (TotalMatchesFound)
+	{
+	case 1:
+		return AddConnection(FromNodeHandle, FromOutputName, ToNodeHandle, ToInputName);
+	case 0:
+		return EGeometryFlowResult::NoMatchesFound;
+	default:
+		return EGeometryFlowResult::MultipleMatchingAmbiguityFound;
+	}
+}
+
+
+
+
+
+EGeometryFlowResult FGraph::FindConnectionForInput(FHandle ToNode, FString ToInput, FConnection& ConnectionOut) const
+{
+	for (const FConnection& Connection : Connections)
+	{
+		if (Connection.ToNode == ToNode && Connection.ToInput == ToInput)
+		{
+			ConnectionOut = Connection;
+			return EGeometryFlowResult::Ok;
+		}
+	}
+	ensure(false);
+	return EGeometryFlowResult::ConnectionDoesNotExist;
+}
+
+
+int32 FGraph::CountOutputConnections(FHandle FromNode, const FString& OutputName) const
+{
+	int32 Count = 0;
+	for (const FConnection& Connection : Connections)
+	{
+		if (Connection.FromNode == FromNode && Connection.FromOutput == OutputName)
+		{
+			Count++;
+		}
+	}
+	return Count;
+}
+
+
+
+TSafeSharedPtr<IData> FGraph::ComputeOutputData(
+	FHandle NodeHandle, 
+	FString OutputName, 
+	TUniquePtr<FEvaluationInfo>& EvaluationInfo,
+	bool bStealOutputData)
+{
+	TSafeSharedPtr<FNode> Node = FindNode(NodeHandle);
+	check(Node);
+
+	// figure out which upstream Connections/Inputs we need to compute this Output
+	TArray<FString> Outputs;
+	Outputs.Add(OutputName);
+	TArray<FEvalRequirement> InputRequirements;
+	Node->CollectRequirements({ OutputName }, InputRequirements);
+
+	// this is the map of (InputName, Data) we will build up by pulling from the Connections
+	FNamedDataMap DataIn;
+
+	// Collect data from those Inputs.
+	// This will recursively call ComputeOutputData() on those (Node/Output) pairs
+	for (int32 k = 0; k < InputRequirements.Num(); ++k)
+	{
+		FDataFlags DataFlags;
+		const FString& InputName = InputRequirements[k].InputName;
+
+		// find the connection for this input
+		FConnection Connection;
+		EGeometryFlowResult FoundResult = FindConnectionForInput(NodeHandle, InputName, Connection);
+		check(FoundResult == EGeometryFlowResult::Ok);
+
+		ENodeCachingStrategy FromCachingStrategy = GetCachingStrategyForNode(Connection.FromNode);
+
+		// If there is only one Connection from this upstream Output (ie to our Input), and the Node/Input 
+		// can steal and transform that data, then we will do it
+		int32 OutputUsageCount = CountOutputConnections(Connection.FromNode, Connection.FromOutput);
+		bool bStealDataForInput = false;
+		if (OutputUsageCount == 1 
+			&& InputRequirements[k].InputFlags.bCanTransformInput
+			&& FromCachingStrategy != ENodeCachingStrategy::AlwaysCache)
+		{
+			bStealDataForInput = true;
+			DataFlags.bIsMutableData = true;
+		}
+
+		// recursively fetch Data coming in to this Input via Connection
+		TSafeSharedPtr<IData> OutputData = 
+			ComputeOutputData(Connection.FromNode, Connection.FromOutput, EvaluationInfo, bStealDataForInput);
+	
+		// add to named data map
+		DataIn.Add(InputName, OutputData, DataFlags);
+	}
+
+	// evalute node
+	FNamedDataMap DataOut;
+	DataOut.Add(OutputName);
+	Node->Evaluate(DataIn, DataOut, EvaluationInfo);
+	EvaluationInfo->CountEvaluation(Node.Get());
+
+	// collect (and optionally take/steal) desired Output data
+	TSafeSharedPtr<IData> Result = (bStealOutputData) ? 
+		Node->StealOutput(OutputName) : DataOut.FindData(OutputName);
+
+	check(Result);
+	return Result;
+}
+
+
+
+
+
+void FGraph::ConfigureCachingStrategy(ENodeCachingStrategy NewStrategy)
+{
+	if (NewStrategy != DefaultCachingStrategy && ensure(NewStrategy != ENodeCachingStrategy::Default) )
+	{
+		DefaultCachingStrategy = NewStrategy;
+
+		// todo: clear caches if necessary?
+	}
+}
+
+
+EGeometryFlowResult FGraph::SetNodeCachingStrategy(FHandle NodeHandle, ENodeCachingStrategy Strategy)
+{
+	FNodeInfo* Found = AllNodes.Find(NodeHandle.Identifier);
+	if (!ensure(Found != nullptr))
+	{
+		return EGeometryFlowResult::NodeDoesNotExist;
+	}
+	Found->CachingStrategy = Strategy;
+	return EGeometryFlowResult::Ok;
+}
\ No newline at end of file
diff --git a/Engine/Plugins/Experimental/GeometryFlow/Source/GeometryFlowCore/Private/GeometryFlowNode.cpp b/Engine/Plugins/Experimental/GeometryFlow/Source/GeometryFlowCore/Private/GeometryFlowNode.cpp
new file mode 100644
--- /dev/null
+++ b/Engine/Plugins/Experimental/GeometryFlow/Source/GeometryFlowCore/Private/GeometryFlowNode.cpp
@@ -0,0 +1,246 @@
+// Copyright Epic Games, Inc. All Rights Reserved.
+
+#include "GeometryFlowNode.h"
+
+using namespace UE::GeometryFlow;
+
+
+void FNode::AddInput(FString Name, TUniquePtr<INodeInput>&& Input)
+{
+	// todo checks/etc
+
+	FNodeInputInfo NodeInputInfo;
+	NodeInputInfo.Name = Name;
+	NodeInputInfo.Input = MoveTemp(Input);
+	NodeInputs.Add(MoveTemp(NodeInputInfo));
+}
+void FNode::AddOutput(FString Name, TUniquePtr<INodeOutput>&& Output)
+{
+	// todo checks/etc
+
+	FNodeOutputInfo NodeOutputInfo;
+	NodeOutputInfo.Name = Name;
+	NodeOutputInfo.Output = MoveTemp(Output);
+	NodeOutputs.Add(MoveTemp(NodeOutputInfo));
+}
+
+
+EGeometryFlowResult FNode::GetInputType(FString Name, int32& TypeOut) const
+{
+	int32 Index = NodeInputs.IndexOfByPredicate([&](const FNodeInputInfo& Input) { return Input.Name == Name; });
+	if (ensure(Index != INDEX_NONE))
+	{
+		TypeOut = NodeInputs[Index].Input->GetDataType();
+		return EGeometryFlowResult::Ok;
+	}
+	return EGeometryFlowResult::InputDoesNotExist;
+}
+
+FNodeInputFlags FNode::GetInputFlags(FString InputName) const
+{
+	int32 Index = NodeInputs.IndexOfByPredicate([&](const FNodeInputInfo& Input) { return Input.Name == InputName; });
+	if (ensure(Index != INDEX_NONE))
+	{
+		return NodeInputs[Index].Input->GetInputFlags();
+	}
+	return FNodeInputFlags();
+}
+
+EGeometryFlowResult FNode::GetOutputType(FString Name, int32& TypeOut) const
+{
+	int32 Index = NodeOutputs.IndexOfByPredicate([&](const FNodeOutputInfo& Output) { return Output.Name == Name; });
+	if (ensure(Index != INDEX_NONE))
+	{
+		TypeOut = NodeOutputs[Index].Output->GetDataType();
+		return EGeometryFlowResult::Ok;
+	}
+	return EGeometryFlowResult::OutputDoesNotExist;
+}
+
+
+void FNode::EnumerateInputs(TFunctionRef<void(const FString& Name, const TUniquePtr<INodeInput>& Input)> EnumerateFunc) const
+{
+	for (const FNodeInputInfo& InputInfo : NodeInputs)
+	{
+		EnumerateFunc(InputInfo.Name, InputInfo.Input);
+	}
+
+}
+
+void FNode::EnumerateOutputs(TFunctionRef<void(const FString& Name, const TUniquePtr<INodeOutput>& Output)> EnumerateFunc) const
+{
+	for (const FNodeOutputInfo& OutputInfo : NodeOutputs)
+	{
+		EnumerateFunc(OutputInfo.Name, OutputInfo.Output);
+	}
+}
+
+
+bool FNode::IsOutputAvailable(FString OutputName) const
+{
+	for (const FNodeOutputInfo& OutputInfo : NodeOutputs)
+	{
+		if (OutputInfo.Name == OutputName)
+		{
+			return OutputInfo.Output->HasCachedOutput();
+		}
+	}
+	ensure(false);		// not found
+	return false;
+}
+
+
+TSafeSharedPtr<IData> FNode::StealOutput(FString OutputName)
+{
+	for (const FNodeOutputInfo& OutputInfo : NodeOutputs)
+	{
+		if (OutputInfo.Name == OutputName)
+		{
+			return OutputInfo.Output->StealOutput();
+		}
+	}
+	ensure(false);
+	return TSafeSharedPtr<IData>();
+}
+
+
+void FNode::CollectRequirements(const TArray<FString>& Outputs, TArray<FEvalRequirement>& RequiredInputsOut)
+{
+	RequiredInputsOut.Reset();
+
+	int32 FoundOutputs = 0;
+	for (const FNodeOutputInfo& OutputInfo : NodeOutputs)
+	{
+		if (Outputs.Contains(OutputInfo.Name))
+		{
+			FoundOutputs++;
+		}
+	}
+	if (ensure(FoundOutputs != 0) == false)
+	{
+		return;
+	}
+
+	for ( const FNodeInputInfo & InputInfo : NodeInputs )
+	{
+		FEvalRequirement Requirement(InputInfo.Name, InputInfo.Input->GetInputFlags());
+		RequiredInputsOut.Add(Requirement);
+	}
+}
+
+
+
+
+void FNode::ConfigureInputFlags(const FString& InputName, FNodeInputFlags Flags)
+{
+	int32 Index = NodeInputs.IndexOfByPredicate([&](const FNodeInputInfo& Input) { return Input.Name == InputName; });
+	if (ensure(Index != INDEX_NONE))
+	{
+		return NodeInputs[Index].Input->SetInputFlags(Flags);
+	}
+}
+
+
+bool FNode::IsInputDirty(FString Name, int32 NewTimestamp) const
+{
+	int32 Index = NodeInputs.IndexOfByPredicate([&](const FNodeInputInfo& Input) { return Input.Name == Name; });
+	if (ensure(Index != INDEX_NONE))
+	{
+		return NodeInputs[Index].LastTimestamp != NewTimestamp;
+	}
+	return true;
+}
+
+bool FNode::CheckIsInputDirtyAndUpdate(FString Name, int32 NewTimestamp)
+{
+	bool bDirty = true;
+	int32 Index = NodeInputs.IndexOfByPredicate([&](const FNodeInputInfo& Input) { return Input.Name == Name; });
+	if (ensure(Index != INDEX_NONE))
+	{
+		bDirty = NodeInputs[Index].LastTimestamp != NewTimestamp;
+		if (bDirty)
+		{
+			NodeInputs[Index].LastTimestamp = NewTimestamp;
+		}
+	}
+	return bDirty;
+}
+
+
+void FNode::UpdateInputTimestamp(FString InputName, int32 NewTimestamp)
+{
+	int32 Index = NodeInputs.IndexOfByPredicate([&](const FNodeInputInfo& Input) { return Input.Name == InputName; });
+	if (ensure(Index != INDEX_NONE))
+	{
+		NodeInputs[Index].LastTimestamp = NewTimestamp;
+	}
+}
+
+void FNode::SetOutput(
+	const FString& OutputName,
+	TSafeSharedPtr<IData> NewData
+)
+{
+	int32 Index = NodeOutputs.IndexOfByPredicate([&](const FNodeOutputInfo& Output) { return Output.Name == OutputName; });
+	if (ensure(Index != INDEX_NONE))
+	{
+		NodeOutputs[Index].Output->UpdateOutput(NewData);
+	}
+}
+
+
+TSafeSharedPtr<IData> FNode::GetOutput(const FString& OutputName) const
+{
+	int32 Index = NodeOutputs.IndexOfByPredicate([&](const FNodeOutputInfo& Output) { return Output.Name == OutputName; });
+	if (ensure(Index != INDEX_NONE))
+	{
+		return NodeOutputs[Index].Output->GetOutput();
+	}
+	return nullptr;
+}
+
+
+TSafeSharedPtr<IData> FNode::FindAndUpdateInputForEvaluate(const FString& InputName, const FNamedDataMap& DatasIn,
+	bool& bAccumModifiedOut, bool& bAccumValidOut)
+{
+	if (ensure(DatasIn.Contains(InputName)) == false)
+	{
+		bAccumValidOut = false;
+		return TSafeSharedPtr<IData>();
+	}
+	else
+	{
+		TSafeSharedPtr<IData> Argument = DatasIn.FindData(InputName);
+		bool bInputModified = CheckIsInputDirtyAndUpdate(InputName, Argument->GetTimestamp());
+		bAccumModifiedOut |= bInputModified;
+		return Argument;
+	}
+}
+
+
+
+
+
+
+
+FEvaluationInfo::FEvaluationInfo()
+{
+	EvaluationsCount = 0;
+	ComputesCount = 0;
+}
+
+void FEvaluationInfo::Reset()
+{
+	EvaluationsCount = 0;
+	ComputesCount = 0;
+}
+
+void FEvaluationInfo::CountEvaluation(FNode* Node)
+{
+	EvaluationsCount++;
+}
+
+void FEvaluationInfo::CountCompute(FNode* Node)
+{
+	ComputesCount++;
+}
\ No newline at end of file
diff --git a/Engine/Plugins/Experimental/GeometryFlow/Source/GeometryFlowCore/Public/BaseNodes/BinaryOpLambdaNode.h b/Engine/Plugins/Experimental/GeometryFlow/Source/GeometryFlowCore/Public/BaseNodes/BinaryOpLambdaNode.h
new file mode 100644
--- /dev/null
+++ b/Engine/Plugins/Experimental/GeometryFlow/Source/GeometryFlowCore/Public/BaseNodes/BinaryOpLambdaNode.h
@@ -0,0 +1,82 @@
+// Copyright Epic Games, Inc. All Rights Reserved.
+
+#pragma once
+
+#include "GeometryFlowNode.h"
+#include "GeometryFlowMovableData.h"
+
+
+namespace UE
+{
+namespace GeometryFlow
+{
+
+
+
+template<typename T, int StorageTypeIdentifier>
+class TBinaryOpLambdaNode : public FNode
+{
+protected:
+	using DataType = TMovableData<T, StorageTypeIdentifier>;
+
+	TUniqueFunction<T(const T& A, const T& B)> OpFunc;
+
+public:
+	static const FString InParamArg1() { return TEXT("Operand1"); }
+	static const FString InParamArg2() { return TEXT("Operand2"); }
+	static const FString OutParamResult() { return TEXT("Result"); }
+
+protected:
+	TBinaryOpLambdaNode()
+	{
+		AddInput(InParamArg1(), MakeUnique<TBasicNodeInput<T, StorageTypeIdentifier>>());
+		AddInput(InParamArg2(), MakeUnique<TBasicNodeInput<T, StorageTypeIdentifier>>());
+
+		AddOutput(OutParamResult(), MakeUnique<TBasicNodeOutput<T, StorageTypeIdentifier>>());
+	}
+
+public:
+	TBinaryOpLambdaNode(TUniqueFunction<T(const T& A, const T& B)>&& OpFuncIn) : TBinaryOpLambdaNode()
+	{
+		OpFunc = MoveTemp(OpFuncIn);
+	}
+
+	virtual void Evaluate(
+		const FNamedDataMap& DatasIn,
+		FNamedDataMap& DatasOut,
+		TUniquePtr<FEvaluationInfo>& EvaluationInfo) override
+	{
+		if (ensure(DatasOut.Contains(OutParamResult())))
+		{
+			bool bAnyInputsModified = false, bAllInputsValid = true;
+			TSafeSharedPtr<IData> Arg1 = FindAndUpdateInputForEvaluate(InParamArg1(), DatasIn, bAnyInputsModified, bAllInputsValid);
+			TSafeSharedPtr<IData> Arg2 = FindAndUpdateInputForEvaluate(InParamArg2(), DatasIn, bAnyInputsModified, bAllInputsValid);
+			if (bAllInputsValid)
+			{
+				if (bAnyInputsModified)
+				{
+					T A, B;
+					Arg1->GetDataCopy(A, StorageTypeIdentifier);
+					Arg2->GetDataCopy(B, StorageTypeIdentifier);
+
+					T BinaryOpResult = OpFunc(A, B);
+
+					TSafeSharedPtr<DataType> Result = MakeSafeShared<DataType>();
+					Result->SetData(BinaryOpResult);
+					SetOutput(ResultName(), Result);
+
+					EvaluationInfo->CountCompute(this);
+				}
+				DatasOut.SetData(OutParamResult(), GetOutput(OutParamResult()) );
+			}
+		}
+	}
+};
+
+
+
+
+
+
+}	// end namespace GeometryFlow
+}	// end namespace UE
diff --git a/Engine/Plugins/Experimental/GeometryFlow/Source/GeometryFlowCore/Public/BaseNodes/TransferNode.h b/Engine/Plugins/Experimental/GeometryFlow/Source/GeometryFlowCore/Public/BaseNodes/TransferNode.h
new file mode 100644
--- /dev/null
+++ b/Engine/Plugins/Experimental/GeometryFlow/Source/GeometryFlowCore/Public/BaseNodes/TransferNode.h
@@ -0,0 +1,80 @@
+// Copyright Epic Games, Inc. All Rights Reserved.
+
+#pragma once
+
+#include "GeometryFlowNode.h"
+
+
+namespace UE
+{
+namespace GeometryFlow
+{
+
+
+template<typename T, int StorageTypeIdentifier>
+class TTransferNode : public FNode
+{
+protected:
+	using DataContainerType = TMovableData<T, StorageTypeIdentifier>;
+
+public:
+	static const FString InParamValue() { return TEXT("Input"); }
+	static const FString OutParamValue() { return TEXT("Output"); }
+
+public:
+	TTransferNode()
+	{
+		AddInput(InParamValue(), MakeUnique<TBasicNodeInput<T, StorageTypeIdentifier>>());
+		AddOutput(OutParamValue(), MakeUnique<TBasicNodeOutput<T, StorageTypeIdentifier>>());
+
+		// we can steal input and pass to output
+		ConfigureInputFlags(InParamValue(), FNodeInputFlags::Transformable());
+	}
+
+	virtual void Evaluate(
+		const FNamedDataMap& DatasIn,
+		FNamedDataMap& DatasOut,
+		TUniquePtr<FEvaluationInfo>& EvaluationInfo) override
+	{
+		if (ensure(DatasOut.Contains(OutParamValue())))
+		{
+			bool bAllInputsValid = true;
+			bool bRecomputeRequired = ( IsOutputAvailable(OutParamValue()) == false );
+			TSafeSharedPtr<IData> InputArg = FindAndUpdateInputForEvaluate(InParamValue(), DatasIn, bRecomputeRequired, bAllInputsValid);
+			if (bAllInputsValid)
+			{
+				if (bRecomputeRequired)
+				{
+					bool bIsDataMutable = DatasIn.GetDataFlags(InParamValue()).bIsMutableData;
+
+					T NewData;
+					if (bIsDataMutable)
+					{
+						InputArg->GiveTo<T>(NewData, StorageTypeIdentifier);
+
+					}
+					else
+					{
+						InputArg->GetDataCopy(NewData, StorageTypeIdentifier);
+					}
+
+					// store new result
+					TSafeSharedPtr<DataContainerType> Result = MakeSafeShared<DataContainerType>();
+					Result->MoveData(MoveTemp(NewData));
+					SetOutput(OutParamValue(), Result);
+
+					EvaluationInfo->CountCompute(this);
+				}
+				DatasOut.SetData(OutParamValue(), GetOutput(OutParamValue()));
+			}
+		}
+	}
+};
+
+
+
+
+}	// end namespace GeometryFlow
+}	// end namespace UE
+
+
diff --git a/Engine/Plugins/Experimental/GeometryFlow/Source/GeometryFlowCore/Public/BaseNodes/TransformerWithSettingsNode.h b/Engine/Plugins/Experimental/GeometryFlow/Source/GeometryFlowCore/Public/BaseNodes/TransformerWithSettingsNode.h
new file mode 100644
--- /dev/null
+++ b/Engine/Plugins/Experimental/GeometryFlow/Source/GeometryFlowCore/Public/BaseNodes/TransformerWithSettingsNode.h
@@ -0,0 +1,114 @@
+// Copyright Epic Games, Inc. All Rights Reserved.
+
+#pragma once
+
+#include "GeometryFlowNode.h"
+#include "GeometryFlowMovableData.h"
+
+
+namespace UE
+{
+namespace GeometryFlow
+{
+
+
+template<
+	typename InputType, 
+	int InputDataTypeIdentifier,
+	typename SettingsType, 
+	int SettingsDataTypeIdentifier,
+	typename OutputType,
+	int OutputDataTypeIdentifier>
+class TTransformerWithSettingsNode : public FNode
+{
+protected:
+	using OutputDataType = TMovableData<OutputType, OutputDataTypeIdentifier>;
+
+public:
+	static const FString InParamInput() { return TEXT("Input"); }
+	static const FString InParamSettings() { return TEXT("Settings"); }
+	static const FString OutParamResult() { return TEXT("Result"); }
+
+public:
+	TTransformerWithSettingsNode()
+	{
+		AddInput(InParamInput(), MakeUnique<TBasicNodeInput<InputType, InputDataTypeIdentifier>>());
+		AddInput(InParamSettings(), MakeUnique<TBasicNodeInput<SettingsType, SettingsDataTypeIdentifier>>());
+
+		AddOutput(OutParamResult(), MakeUnique<TBasicNodeOutput<OutputType, OutputDataTypeIdentifier>>());
+	}
+
+
+protected:
+
+	//
+	// TTransformerWithSettingsNode API that subclasses must/can implement
+	//
+
+	virtual void ComputeOutput(
+		const FNamedDataMap& DatasIn,
+		const SettingsType& SettingsIn,
+		const InputType& InputIn,
+		OutputType& OutputOut)
+	{
+		check(false);		// should never be here, subclass needs to override
+	}
+
+	virtual void CheckAdditionalInputs(const FNamedDataMap& DatasIn, bool& bRecomputeRequired, bool& bAllInputsValid)
+	{
+		// none
+	}
+
+
+public:
+
+
+	virtual void Evaluate(
+		const FNamedDataMap& DatasIn,
+		FNamedDataMap& DatasOut,
+		TUniquePtr<FEvaluationInfo>& EvaluationInfo) override
+	{
+		if (ensure(DatasOut.Contains(OutParamResult())))
+		{
+			bool bAllInputsValid = true;
+			bool bRecomputeRequired = (IsOutputAvailable(OutParamResult()) == false);
+			TSafeSharedPtr<IData> InputArg = FindAndUpdateInputForEvaluate(InParamInput(), DatasIn, bRecomputeRequired, bAllInputsValid);
+			TSafeSharedPtr<IData> SettingsArg = FindAndUpdateInputForEvaluate(InParamSettings(), DatasIn, bRecomputeRequired, bAllInputsValid);
+			CheckAdditionalInputs(DatasIn, bRecomputeRequired, bAllInputsValid);
+			if (bAllInputsValid)
+			{
+				if (bRecomputeRequired)
+				{
+					const InputType& SourceData = InputArg->GetDataConstRef<InputType>(InputDataTypeIdentifier);
+
+					// always make a copy of settings
+					SettingsType Settings;
+					SettingsArg->GetDataCopy(Settings, SettingsDataTypeIdentifier);
+
+					// run processing
+					OutputType ResultData;
+					ComputeOutput(DatasIn, Settings, SourceData, ResultData);
+
+					// store new result
+					TSafeSharedPtr<OutputDataType> Result = MakeSafeShared<OutputDataType>();
+					Result->MoveData(MoveTemp(ResultData));
+					SetOutput(OutParamResult(), Result);
+
+					EvaluationInfo->CountCompute(this);
+				}
+				DatasOut.SetData(OutParamResult(), GetOutput(OutParamResult()));
+			}
+		}
+	}
+
+
+
+};
+
+
+
+
+
+
+}	// end namespace GeometryFlow
+}	// end namespace UE
\ No newline at end of file
diff --git a/Engine/Plugins/Experimental/GeometryFlow/Source/GeometryFlowCore/Public/GeometryFlowCoreModule.h b/Engine/Plugins/Experimental/GeometryFlow/Source/GeometryFlowCore/Public/GeometryFlowCoreModule.h
new file mode 100644
--- /dev/null
+++ b/Engine/Plugins/Experimental/GeometryFlow/Source/GeometryFlowCore/Public/GeometryFlowCoreModule.h
@@ -0,0 +1,15 @@
+// Copyright Epic Games, Inc. All Rights Reserved.
+
+#pragma once
+
+#include "CoreMinimal.h"
+#include "Modules/ModuleManager.h"
+
+class FGeometryFlowCoreModule : public IModuleInterface
+{
+public:
+
+	/** IModuleInterface implementation */
+	virtual void StartupModule() override;
+	virtual void ShutdownModule() override;
+};
diff --git a/Engine/Plugins/Experimental/GeometryFlow/Source/GeometryFlowCore/Public/GeometryFlowCoreNodes.h b/Engine/Plugins/Experimental/GeometryFlow/Source/GeometryFlowCore/Public/GeometryFlowCoreNodes.h
new file mode 100644
--- /dev/null
+++ b/Engine/Plugins/Experimental/GeometryFlow/Source/GeometryFlowCore/Public/GeometryFlowCoreNodes.h
@@ -0,0 +1,114 @@
+// Copyright Epic Games, Inc. All Rights Reserved.
+
+#pragma once
+
+#include "GeometryFlowNode.h"
+#include "GeometryFlowMovableData.h"
+
+
+namespace UE
+{
+namespace GeometryFlow
+{
+
+
+
+template<typename T, int32 StorageTypeIdentifier>
+class TSourceNode : public FNode
+{
+protected:
+	using DataType = TMovableData<T, StorageTypeIdentifier>;
+
+	TSafeSharedPtr<DataType> Value;
+
+public:
+	using CppType = T;
+	static constexpr int DataTypeIdentifier = StorageTypeIdentifier;
+
+	static const FString OutParamValue() { return TEXT("Value"); }
+
+public:
+	TSourceNode()
+	{
+		T InitialValue;
+		Value = MakeSafeShared<DataType>(InitialValue);
+		AddOutput(OutParamValue(), MakeUnique<TBasicNodeOutput<T, StorageTypeIdentifier>>());
+
+		UpdateSourceValue(InitialValue);
+	}
+
+	virtual void UpdateSourceValue(const T& NewValue)
+	{
+		Value->SetData(NewValue);
+		SetOutput(OutParamValue(), Value);
+	}
+
+	virtual void CollectRequirements(const TArray<FString>& Outputs, TArray<FEvalRequirement>& RequiredInputsOut) override
+	{
+		return;
+	}
+
+	virtual void Evaluate(
+		const FNamedDataMap& DatasIn,
+		FNamedDataMap& DatasOut,
+		TUniquePtr<FEvaluationInfo>& EvaluationInfo) override
+	{
+		if ( ensure(DatasOut.Contains(OutParamValue())) )
+		{
+			DatasOut.SetData(OutParamValue(), GetOutput(OutParamValue()));
+		}
+	}
+};
+
+
+
+
+/**
+ * This macro declares a MovableData, Node Input/Output, and TSourceNode type for a given C++ type/identifier
+ * 
+ * eg GEOMETRYFLOW_DECLARE_BASIC_TYPES(Int32, int, (int)EDataTypes::Integer) will declare/define:
+ *     - TMovableData type named FDataInt32
+ *     - TBasicNodeInput type named FInt32Input
+ *     - TBasicNodeOutput type named FInt32Output
+ *     - TSourceNode type named FInt32SourceNode
+ *
+ */
+#define GEOMETRYFLOW_DECLARE_BASIC_TYPES(TypeName, CppType, TypeIdentifier) \
+	typedef TMovableData<CppType, TypeIdentifier> FData##TypeName; \
+	typedef TBasicNodeInput<CppType, TypeIdentifier> F##TypeName##Input; \
+	typedef TBasicNodeOutput<CppType, TypeIdentifier> F##TypeName##Output; \
+	typedef TSourceNode<CppType, TypeIdentifier> F##TypeName##SourceNode; 
+
+
+
+
+
+/**
+ * This macro defines two classes for a given C++ class:
+ *   - TMovableData typedef named FData[XYZ]Settings
+ *   - TSourceNode typedef named F[XYZ]SettingsSourceNode
+ *
+ * The assumption is the class has an integer static/constexpr member SettingsType::DataTypeIdentifier that defines the type integer
+ */
+#define GEOMETRYFLOW_DECLARE_SETTINGS_TYPES(CppType, ReadableName) \
+	typedef TMovableData<CppType, CppType::DataTypeIdentifier> FData##ReadableName##Settings; \
+	typedef TSourceNode<CppType, CppType::DataTypeIdentifier> F##ReadableName##SettingsSourceNode;
+
+
+
+
+
+
+/**
+ * Declare basic types for math
+ */
+GEOMETRYFLOW_DECLARE_BASIC_TYPES(Int32, int, (int)EDataTypes::Integer)
+GEOMETRYFLOW_DECLARE_BASIC_TYPES(Float, float, (int)EDataTypes::Float)
+GEOMETRYFLOW_DECLARE_BASIC_TYPES(Double, double, (int)EDataTypes::Double)
+GEOMETRYFLOW_DECLARE_BASIC_TYPES(Vector3f, FVector3f, (int)EDataTypes::Vector3f)
+GEOMETRYFLOW_DECLARE_BASIC_TYPES(Vector3d, FVector3d, (int)EDataTypes::Vector3d)
+
+
+
+}	// end namespace GeometryFlow
+}	// end namespace UE
\ No newline at end of file
diff --git a/Engine/Plugins/Experimental/GeometryFlow/Source/GeometryFlowCore/Public/GeometryFlowGraph.h b/Engine/Plugins/Experimental/GeometryFlow/Source/GeometryFlowCore/Public/GeometryFlowGraph.h
new file mode 100644
--- /dev/null
+++ b/Engine/Plugins/Experimental/GeometryFlow/Source/GeometryFlowCore/Public/GeometryFlowGraph.h
@@ -0,0 +1,175 @@
+// Copyright Epic Games, Inc. All Rights Reserved.
+
+#pragma once
+
+#include "CoreMinimal.h"
+#include "GeometryFlowNode.h"
+
+
+namespace UE
+{
+namespace GeometryFlow
+{
+
+
+
+// TODO: this needs more work
+	// need option to opportunistically cache, ie on multi-use of output, or when we can't re-use anyway
+	//   (this is currently what NeverCache does?)
+
+enum class ENodeCachingStrategy
+{
+	Default = 0,
+	AlwaysCache = 1,
+	NeverCache = 2
+};
+
+
+//
+// TODO: 
+// - internal FNode pointers can be unique ptr?
+// - parallel evaluation at graph level (pre-pass to collect references?)
+//
+
+
+class GEOMETRYFLOWCORE_API FGraph
+{
+public:
+
+	struct FHandle
+	{
+		int32 Identifier;
+		bool operator==(const FHandle& OtherHandle) const { return Identifier == OtherHandle.Identifier; }
+	};
+
+	template<typename NodeType>
+	FHandle AddNodeOfType(
+		const FString& Identifier = FString(""),
+		ENodeCachingStrategy CachingStrategy = ENodeCachingStrategy::Default )
+	{
+		FNodeInfo NewNodeInfo;
+		NewNodeInfo.Node = MakeShared<NodeType, ESPMode::ThreadSafe>();
+		NewNodeInfo.Node->SetIdentifier(Identifier);
+		NewNodeInfo.CachingStrategy = CachingStrategy;
+		FHandle Handle = { NodeCounter++ };
+		AllNodes.Add(Handle.Identifier, NewNodeInfo);
+		return Handle;
+	}
+
+	EGeometryFlowResult AddConnection(FHandle FromNode, FString FromOutput, FHandle ToNode, FString ToInput);
+
+	EGeometryFlowResult InferConnection(FHandle FromNode, FHandle ToNode);
+
+
+	template<typename T>
+	EGeometryFlowResult EvaluateResult(FHandle Node, FString OutputName, T& Storage, int32 StorageTypeIdentifier)
+	{
+		TUniquePtr<FEvaluationInfo> EvalInfo = MakeUnique<FEvaluationInfo>();
+		return EvaluateResult(Node, OutputName, Storage, StorageTypeIdentifier, EvalInfo);
+	}
+
+	template<typename T>
+	EGeometryFlowResult EvaluateResult(
+		FHandle Node, 
+		FString OutputName, 
+		T& Storage, 
+		int32 StorageTypeIdentifier,
+		TUniquePtr<FEvaluationInfo>& EvaluationInfo,
+		bool bTryTakeResult)
+	{
+		int32 OutputType;
+		EGeometryFlowResult ValidOutput = GetOutputTypeForNode(Node, OutputName, OutputType);
+		if (ValidOutput != EGeometryFlowResult::Ok)
+		{
+			return ValidOutput;
+		}
+		if (OutputType != StorageTypeIdentifier)
+		{
+			return EGeometryFlowResult::UnmatchedTypes;
+		}
+		if (bTryTakeResult)
+		{
+			TSafeSharedPtr<IData> Data = ComputeOutputData(Node, OutputName, EvaluationInfo, true);
+			Data->GiveTo(Storage, StorageTypeIdentifier);
+		}
+		else
+		{
+			TSafeSharedPtr<IData> Data = ComputeOutputData(Node, OutputName, EvaluationInfo, false);
+			Data->GetDataCopy(Storage, StorageTypeIdentifier);
+		}
+		return EGeometryFlowResult::Ok;
+	}
+
+
+	template<typename NodeType>
+	EGeometryFlowResult ApplyToNodeOfType(
+		FHandle NodeHandle, 
+		TFunctionRef<void(NodeType&)> ApplyFunc)
+	{
+		TSafeSharedPtr<FNode> FoundNode = FindNode(NodeHandle);
+		if (FoundNode)
+		{
+			NodeType* Value = static_cast<NodeType*>(FoundNode.Get());
+			ApplyFunc(*Value);
+			return EGeometryFlowResult::Ok;
+		}
+		return EGeometryFlowResult::NodeDoesNotExist;
+	}
+
+
+
+	void ConfigureCachingStrategy(ENodeCachingStrategy NewStrategy);
+	EGeometryFlowResult SetNodeCachingStrategy(FHandle NodeHandle, ENodeCachingStrategy Strategy);
+
+	
+protected:
+	int32 NodeCounter = 0;
+
+	ENodeCachingStrategy DefaultCachingStrategy = ENodeCachingStrategy::AlwaysCache;
+
+	struct FNodeInfo
+	{
+		TSafeSharedPtr<FNode> Node;
+		ENodeCachingStrategy CachingStrategy = ENodeCachingStrategy::Default;
+	};
+
+	TMap<int32, FNodeInfo> AllNodes;
+
+	TSafeSharedPtr<FNode> FindNode(FHandle Handle) const;
+	EGeometryFlowResult GetInputTypeForNode(FHandle NodeHandle, FString InputName, int32& Type) const;
+	EGeometryFlowResult GetOutputTypeForNode(FHandle NodeHandle, FString OutputName, int32& Type) const;
+	ENodeCachingStrategy GetCachingStrategyForNode(FHandle NodeHandle) const;
+
+	struct FConnection
+	{
+		FHandle FromNode;
+		FString FromOutput;
+
+		FHandle ToNode;
+		FString ToInput;
+	};
+	TArray<FConnection> Connections;
+
+	EGeometryFlowResult FindConnectionForInput(FHandle ToNode, FString ToInput, FConnection& ConnectionOut) const;
+
+	int32 CountOutputConnections(FHandle FromNode, const FString& FromOutput) const;
+
+	TSafeSharedPtr<IData> ComputeOutputData(
+		FHandle Node, 
+		FString OutputName, 
+		TUniquePtr<FEvaluationInfo>& EvaluationInfo,
+		bool bStealOutputData = false);
+};
+
+
+
+
+
+
+
+
+
+
+
+}	// end namespace GeometryFlow
+}	//
\ No newline at end of file
diff --git a/Engine/Plugins/Experimental/GeometryFlow/Source/GeometryFlowCore/Public/GeometryFlowGraphUtil.h b/Engine/Plugins/Experimental/GeometryFlow/Source/GeometryFlowCore/Public/GeometryFlowGraphUtil.h
new file mode 100644
--- /dev/null
+++ b/Engine/Plugins/Experimental/GeometryFlow/Source/GeometryFlowCore/Public/GeometryFlowGraphUtil.h
@@ -0,0 +1,42 @@
+// Copyright Epic Games, Inc. All Rights Reserved.
+
+#pragma once
+
+#include "GeometryFlowGraph.h"
+#include "GeometryFlowCoreNodes.h"
+
+namespace UE
+{
+namespace GeometryFlow
+{
+
+
+
+template<typename SourceNodeType>
+void UpdateSourceNodeValue(FGraph& Graph, FGraph::FHandle NodeHandle, const typename SourceNodeType::CppType& NewValue)
+{
+	Graph.ApplyToNodeOfType<SourceNodeType>(NodeHandle, [&](SourceNodeType& Node)
+	{
+		Node.UpdateSourceValue(NewValue);
+	});
+}
+
+
+
+
+template<typename SettingsType>
+void UpdateSettingsSourceNodeValue(FGraph& Graph, FGraph::FHandle NodeHandle, const SettingsType& NewSettings)
+{
+	using SettingsSourceNodeType = TSourceNode<SettingsType, SettingsType::DataTypeIdentifier>;
+
+	Graph.ApplyToNodeOfType<SettingsSourceNodeType>(NodeHandle, [&](SettingsSourceNodeType& Node)
+	{
+		Node.UpdateSourceValue(NewSettings);
+	});
+}
+
+
+
+
+}	// end namespace GeometryFlow
+}	// end namespace UE
\ No newline at end of file
diff --git a/Engine/Plugins/Experimental/GeometryFlow/Source/GeometryFlowCore/Public/GeometryFlowImmutableData.h b/Engine/Plugins/Experimental/GeometryFlow/Source/GeometryFlowCore/Public/GeometryFlowImmutableData.h
new file mode 100644
--- /dev/null
+++ b/Engine/Plugins/Experimental/GeometryFlow/Source/GeometryFlowCore/Public/GeometryFlowImmutableData.h
@@ -0,0 +1,110 @@
+// Copyright Epic Games, Inc. All Rights Reserved.
+
+#pragma once
+
+#include "GeometryFlowNode.h"
+
+
+namespace UE
+{
+namespace GeometryFlow
+{
+
+
+
+
+template<typename T, int32 DataTypeIdentifier>
+class TImmutableData : public IData
+{
+public:
+	TUniquePtr<const T> Data;
+
+	TImmutableData()
+	{
+	}
+
+	TImmutableData(TUniquePtr<const T>&& DataIn)
+	{
+		Data = MoveTemp(DataIn);
+	}
+
+	virtual int32 GetPayloadType() const override
+	{
+		return DataTypeIdentifier;
+	}
+
+	virtual int64 GetPayloadBytes() const override
+	{
+		check(false);
+		return 0;
+	}
+
+	virtual bool CopyData(void* StorageType, int32 AsType) const override
+	{
+		check(false);
+		return false;
+	}
+
+	virtual bool MoveDataOut(void* StorageType, int32 AsType) override
+	{
+		check(false);
+		return false;
+	}
+
+protected:
+	
+	// This is used by IData::GetDataConstRef(). Should probably not be
+	// called in ot her contexts...
+	virtual const void* GetRawDataPointerUnsafe() const override
+	{
+		return (const void *)Data.Get();
+	}
+
+
+	const T& GetData() const
+	{
+		return *Data;
+	}
+};
+
+
+
+
+
+template<typename T, int32 StorageTypeIdentifier>
+class TImmutableNodeInput : public INodeInput
+{
+public:
+
+	virtual int32 GetDataType() const { return StorageTypeIdentifier; }
+
+
+public:
+	virtual bool CanTransformInput() const override 
+	{ 
+		ensure(InputFlags.bCanTransformInput == false); 
+		return false;
+	}
+};
+
+
+
+
+template<typename T, int32 StorageTypeIdentifier>
+class TImmutableNodeOutput : public INodeOutput
+{
+public:
+	TImmutableNodeOutput()
+	{
+		TUniquePtr<const T> InitialValue = MakeUnique<const T>();
+		UpdateOutput(MakeShared<TImmutableData<T, StorageTypeIdentifier>, ESPMode::ThreadSafe>(MoveTemp(InitialValue)));
+	}
+
+	virtual int32 GetDataType() const { return StorageTypeIdentifier; }
+};
+
+
+
+
+}	// end namespace GeometryFlow
+}	// end namespace UE
diff --git a/Engine/Plugins/Experimental/GeometryFlow/Source/GeometryFlowCore/Public/GeometryFlowMovableData.h b/Engine/Plugins/Experimental/GeometryFlow/Source/GeometryFlowCore/Public/GeometryFlowMovableData.h
new file mode 100644
--- /dev/null
+++ b/Engine/Plugins/Experimental/GeometryFlow/Source/GeometryFlowCore/Public/GeometryFlowMovableData.h
@@ -0,0 +1,123 @@
+// Copyright Epic Games, Inc. All Rights Reserved.
+
+#pragma once
+
+#include "GeometryFlowNode.h"
+
+
+namespace UE
+{
+namespace GeometryFlow
+{
+
+
+
+
+
+
+
+template<typename T, int32 DataTypeIdentifier>
+class TMovableData : public IData
+{
+public:
+	T Data;
+
+	TMovableData()
+	{
+	}
+
+	TMovableData(const T& InitialValue)
+	{
+		Data = InitialValue;
+	}
+
+	virtual int32 GetPayloadType() const override
+	{
+		return DataTypeIdentifier;
+	}
+
+	virtual int64 GetPayloadBytes() const override
+	{
+		return sizeof(Data);
+	}
+
+	virtual bool CopyData(void* StorageType, int32 AsType) const override
+	{
+		if (ensure(AsType == GetPayloadType()))
+		{
+			*((T*)StorageType) = Data;
+			return true;
+		}
+		return false;
+	}
+
+	virtual bool MoveDataOut(void* StorageType, int32 AsType) override
+	{
+		if (ensure(AsType == GetPayloadType()))
+		{
+			*((T*)StorageType) = MoveTemp(Data);
+			return true;
+		}
+		return false;
+	}
+
+protected:
+	virtual const void* GetRawDataPointerUnsafe() const override
+	{
+		return (const void *)&Data;
+	}
+
+public:
+
+
+	void GetData(T& DataOut) const
+	{
+		DataOut = Data;
+	}
+
+	void SetData(const T& DataIn)
+	{
+		Data = DataIn;
+		IncrementTimestamp();
+	}
+
+	void MoveData(T&& DataIn)
+	{
+		Data = MoveTemp(DataIn);
+		IncrementTimestamp();
+	}
+};
+
+
+
+
+
+template<typename T, int32 StorageTypeIdentifier>
+class TBasicNodeInput : public INodeInput
+{
+public:
+
+	virtual int32 GetDataType() const { return StorageTypeIdentifier; }
+};
+
+
+
+
+template<typename T, int32 StorageTypeIdentifier>
+class TBasicNodeOutput : public INodeOutput
+{
+public:
+	TBasicNodeOutput()
+	{
+		T InitialValue;
+		UpdateOutput(MakeShared<TMovableData<T, StorageTypeIdentifier>, ESPMode::ThreadSafe>(InitialValue));
+	}
+
+	virtual int32 GetDataType() const { return StorageTypeIdentifier; }
+};
+
+
+
+
+}	// end namespace GeometryFlow
+}	// end namespace UE
diff --git a/Engine/Plugins/Experimental/GeometryFlow/Source/GeometryFlowCore/Public/GeometryFlowNode.h b/Engine/Plugins/Experimental/GeometryFlow/Source/GeometryFlowCore/Public/GeometryFlowNode.h
new file mode 100644
--- /dev/null
+++ b/Engine/Plugins/Experimental/GeometryFlow/Source/GeometryFlowCore/Public/GeometryFlowNode.h
@@ -0,0 +1,381 @@
+// Copyright Epic Games, Inc. All Rights Reserved.
+
+#pragma once
+
+#include "CoreMinimal.h"
+#include "HAL/UnrealMemory.h"
+#include "VectorTypes.h"
+#include "GeometryFlowTypes.h"
+
+#include <atomic>
+
+
+// TODO
+// - handling for array types
+// - cancellation support
+// - can we have the concept of data shared between nodes? ie when recomputing normals, we currently 'steal'
+//   but really we are only partially modifying input data...   (maybe cleaner to just separate out compute and set normals?)
+
+
+namespace UE
+{
+namespace GeometryFlow
+{
+
+
+
+class GEOMETRYFLOWCORE_API IData
+{
+public:
+	virtual ~IData() {}
+
+	virtual int32 GetPayloadType() const = 0;
+	virtual int64 GetPayloadBytes() const = 0;
+	virtual bool CopyData(void* StorageType, int32 AsType) const = 0;
+	virtual bool MoveDataOut(void* StorageType, int32 AsType) = 0;
+
+protected:
+	virtual const void* GetRawDataPointerUnsafe() const = 0;
+
+public:
+
+	virtual bool IsSameType(const IData& OtherType)
+	{
+		return GetPayloadType() == OtherType.GetPayloadType();
+	}
+
+	template<typename T>
+	void GetDataCopy(T& DataOut, int32 AsType) const
+	{
+		CopyData((void*)&DataOut, AsType);
+	}
+
+	template<typename T>
+	const T& GetDataConstRef(int32 AsType) const
+	{
+		check(GetPayloadType() == AsType);
+		return *(const T*)GetRawDataPointerUnsafe();
+	}
+
+	template<typename T>
+	void GiveTo(T& DataOut, int32 AsType)
+	{
+		MoveDataOut((void*)&DataOut, AsType);
+	}
+
+public:
+	virtual int32 GetTimestamp() const
+	{
+		return Timestamp;
+	}
+
+protected:
+	int32 Timestamp = 0;
+
+	friend class INodeOutput;
+
+	virtual void IncrementTimestamp()
+	{
+		Timestamp++;
+	}
+
+	virtual void SetMaxTimestamp(int32 ExternalTimestamp)
+	{
+		Timestamp = FMath::Max(ExternalTimestamp, Timestamp);
+	}
+
+};
+
+
+
+struct GEOMETRYFLOWCORE_API FNodeInputFlags
+{
+	bool bCanTransformInput = false;
+
+	static FNodeInputFlags Transformable() { return FNodeInputFlags{ true }; }
+};
+
+
+class GEOMETRYFLOWCORE_API INodeInput
+{
+public:
+	virtual ~INodeInput() {}
+
+	/*
+	 * INodeInput virtual API that implementations must implement
+	 */
+
+	virtual int32 GetDataType() const = 0;
+
+
+
+	/*
+	 * INodeInput control API
+	 */
+
+public:
+	virtual FNodeInputFlags GetInputFlags() const { return InputFlags; }
+	virtual bool CanTransformInput() const { return InputFlags.bCanTransformInput; }
+protected:
+	FNodeInputFlags InputFlags;
+	virtual void SetInputFlags(FNodeInputFlags SetFlags) { InputFlags = SetFlags; }
+
+	friend class FNode;
+};
+
+
+
+
+
+class GEOMETRYFLOWCORE_API INodeOutput
+{
+public:
+	virtual ~INodeOutput() {}
+
+	virtual int32 GetDataType() const = 0;
+
+	virtual void UpdateOutput(TSafeSharedPtr<IData> NewData)
+	{
+		if (ensure(NewData->GetPayloadType() == GetDataType()))
+		{
+			CachedValue = NewData;
+			CachedValue->SetMaxTimestamp(LastDataTimestamp + 1);
+			LastDataTimestamp = CachedValue->GetTimestamp();
+		}
+	}
+
+	virtual bool HasCachedOutput() const
+	{
+		return (CachedValue != nullptr);
+	}
+
+	virtual TSafeSharedPtr<IData> GetOutput() const
+	{
+		check(CachedValue != nullptr);
+		return CachedValue;
+	}
+
+	virtual TSafeSharedPtr<IData> StealOutput()
+	{
+		check(CachedValue != nullptr);
+		TSafeSharedPtr<IData> Result = CachedValue;
+		CachedValue = TSafeSharedPtr<IData>();
+		LastDataTimestamp++;
+		return Result;
+	}
+
+protected:
+	int32 LastDataTimestamp = 0;
+	TSafeSharedPtr<IData> CachedValue;
+};
+
+
+
+struct GEOMETRYFLOWCORE_API FDataFlags
+{
+	bool bIsMutableData = false;
+};
+
+
+class GEOMETRYFLOWCORE_API FNamedDataMap
+{
+protected:
+	TArray<FString> Names;
+	TArray<TSafeSharedPtr<IData>> Datas;
+	TArray<FDataFlags> DatasFlags;
+
+public:
+	void Add(const FString& Name, FDataFlags Flags = FDataFlags())
+	{
+		check(Contains(Name) == false);
+		Names.Add(Name);
+		Datas.Add(TSafeSharedPtr<IData>());
+		DatasFlags.Add(Flags);
+	}
+
+	void Add(const FString& Name, TSafeSharedPtr<IData> Data, FDataFlags Flags = FDataFlags())
+	{
+		check(Contains(Name) == false);
+		Names.Add(Name);
+		Datas.Add(Data);
+		DatasFlags.Add(Flags);
+	}
+
+	bool Contains(const FString& Name) const { return Names.Contains(Name); }
+
+	const TArray<FString>& GetNames() const { return Names; }
+
+	TSafeSharedPtr<IData> FindData(const FString& Name) const
+	{
+		int32 Index = Names.IndexOfByKey(Name);
+		if (ensure(Index != INDEX_NONE))
+		{
+			return Datas[Index];
+		}
+		return nullptr;
+	}
+
+	FDataFlags GetDataFlags(const FString& Name) const
+	{
+		int32 Index = Names.IndexOfByKey(Name);
+		if (ensure(Index != INDEX_NONE))
+		{
+			return DatasFlags[Index];
+		}
+		return FDataFlags();
+	}
+
+	bool SetData(const FString& Name, TSafeSharedPtr<IData> Data)
+	{
+		int32 Index = Names.IndexOfByKey(Name);
+		if (ensure(Index != INDEX_NONE))
+		{
+			Datas[Index] = Data;
+			return true;
+		}
+		return false;
+	}
+};
+
+
+
+
+class GEOMETRYFLOWCORE_API FEvaluationInfo
+{
+public:
+	FEvaluationInfo();
+	virtual ~FEvaluationInfo() {}
+
+	virtual void Reset();
+
+	virtual void CountEvaluation(FNode* Node);
+	int NumEvaluations() const { return EvaluationsCount; }
+
+	virtual void CountCompute(FNode* Node);
+	int NumComputes() const { return ComputesCount; }
+
+protected:
+	std::atomic<int> EvaluationsCount;
+	std::atomic<int> ComputesCount;
+};
+
+
+
+
+struct GEOMETRYFLOWCORE_API FEvalRequirement
+{
+	FString InputName;
+	FNodeInputFlags InputFlags;
+
+	FEvalRequirement() {}
+	FEvalRequirement(const FString& Name) : InputName(Name) {}
+
+	FEvalRequirement(const FString& NameIn, FNodeInputFlags InputFlagsIn)
+		: InputName(NameIn), InputFlags(InputFlagsIn) {}
+};
+
+
+
+class GEOMETRYFLOWCORE_API FNode
+{
+public:
+
+	FNode() = default;
+	FNode(const FNode&) = delete;
+	FNode& operator=(const FNode&) = delete;
+	FNode(FNode&&) = default;
+	FNode& operator=(FNode&&) = default;
+	virtual ~FNode() = default;
+
+	void SetIdentifier(const FString& IdentifierIn) { Identifier = IdentifierIn; }
+	const FString& GetIdentifier() const { return Identifier; }
+
+	void AddInput(FString Name, TUniquePtr<INodeInput>&& Input);
+	void AddOutput(FString Name, TUniquePtr<INodeOutput>&& Output);
+
+	EGeometryFlowResult GetInputType(FString Name, int32& TypeOut) const;
+	EGeometryFlowResult GetOutputType(FString Name, int32& TypeOut) const;
+
+	FNodeInputFlags GetInputFlags(FString InputName) const;
+
+	virtual void EnumerateInputs(TFunctionRef<void(const FString& Name, const TUniquePtr<INodeInput>& Input)> EnumerateFunc) const;
+	virtual void EnumerateOutputs(TFunctionRef<void(const FString& Name, const TUniquePtr<INodeOutput>& Output)> EnumerateFunc) const;
+
+	virtual bool IsOutputAvailable(FString OutputName) const;
+
+	virtual TSafeSharedPtr<IData> StealOutput(FString OutputName);
+
+	/**
+	 * Find the list of named inputs that must be available to compute the named Outputs.
+	 * By default will return all Inputs if any of the listed Outputs exists on this Node.
+	 */
+	virtual void CollectRequirements(
+		const TArray<FString>& Outputs, 
+		TArray<FEvalRequirement>& RequiredInputsOut);
+
+
+	virtual void Evaluate(
+		const FNamedDataMap& DatasIn,
+		FNamedDataMap& DatasOut,
+		TUniquePtr<FEvaluationInfo>& EvaluationInfo) = 0;
+
+protected:
+	FString Identifier;
+
+	struct FNodeInputInfo
+	{
+		// explicitly non-copyable because of embedded TUniquePtr
+		FNodeInputInfo() = default;
+		FNodeInputInfo(const FNodeInputInfo&) = delete;
+		FNodeInputInfo& operator=(const FNodeInputInfo&) = delete;
+		FNodeInputInfo(FNodeInputInfo&&) = default;
+		FNodeInputInfo& operator=(FNodeInputInfo&&) = default;
+
+		FString Name;
+		int32 LastTimestamp = -1;
+		TUniquePtr<INodeInput> Input;
+	};
+
+	struct FNodeOutputInfo
+	{
+		// explicitly non-copyable because of embedded TUniquePtr
+		FNodeOutputInfo() = default;
+		FNodeOutputInfo(const FNodeOutputInfo&) = delete;
+		FNodeOutputInfo& operator=(const FNodeOutputInfo&) = delete;
+		FNodeOutputInfo(FNodeOutputInfo&&) = default;
+		FNodeOutputInfo& operator=(FNodeOutputInfo&&) = default;
+
+		FString Name;
+		TUniquePtr<INodeOutput> Output;
+	};
+
+	virtual void SetOutput(
+		const FString& OutputName,
+		TSafeSharedPtr<IData> NewData
+	);
+	TSafeSharedPtr<IData> GetOutput(const FString& OutputName) const;
+
+	TArray<FNodeInputInfo> NodeInputs;
+	TArray<FNodeOutputInfo> NodeOutputs;
+
+
+protected:
+	// helper functions for setup
+	void ConfigureInputFlags(const FString& InputName, FNodeInputFlags Flags);
+
+	// helper functions for evaluation
+	bool IsInputDirty(FString Name, int32 NewTimestamp) const;
+	bool CheckIsInputDirtyAndUpdate(FString Name, int32 NewTimestamp);
+	void UpdateInputTimestamp(FString Name, int32 NewTimestamp);
+
+	TSafeSharedPtr<IData> FindAndUpdateInputForEvaluate(const FString& InputName, const FNamedDataMap& DatasIn,
+		bool& bAccumModifiedOut, bool& bAccumValidOut);
+};
+
+
+
+
+
+}	// end namespace GeometryFlow
+}	// end namespace UE
+
diff --git a/Engine/Plugins/Experimental/GeometryFlow/Source/GeometryFlowCore/Public/GeometryFlowNodeUtil.h b/Engine/Plugins/Experimental/GeometryFlow/Source/GeometryFlowCore/Public/GeometryFlowNodeUtil.h
new file mode 100644
--- /dev/null
+++ b/Engine/Plugins/Experimental/GeometryFlow/Source/GeometryFlowCore/Public/GeometryFlowNodeUtil.h
@@ -0,0 +1,40 @@
+// Copyright Epic Games, Inc. All Rights Reserved.
+
+#pragma once
+
+#include "GeometryFlowCoreNodes.h"
+#include "GeometryFlowMovableData.h"
+
+namespace UE
+{
+namespace GeometryFlow
+{
+
+
+
+template<typename DataType>
+TSafeSharedPtr<TMovableData<DataType, DataType::DataTypeIdentifier>> MakeMovableData(DataType&& Data)
+{
+	return MakeShared<TMovableData<DataType, DataType::DataTypeIdentifier>, ESPMode::ThreadSafe>(MoveTemp(Data));
+}
+
+
+template<typename DataType>
+TUniquePtr<TBasicNodeInput<DataType, DataType::DataTypeIdentifier>> MakeBasicInput()
+{
+	return MakeUnique<TBasicNodeInput<DataType, DataType::DataTypeIdentifier>>();
+}
+
+template<typename DataType>
+TUniquePtr<TBasicNodeOutput<DataType, DataType::DataTypeIdentifier>> MakeBasicOutput()
+{
+	return MakeUnique<TBasicNodeOutput<DataType, DataType::DataTypeIdentifier>>();
+}
+
+
+
+
+
+
+}	// end namespace GeometryFlow
+}	// end n
\ No newline at end of file
diff --git a/Engine/Plugins/Experimental/GeometryFlow/Source/GeometryFlowCore/Public/GeometryFlowTypes.h b/Engine/Plugins/Experimental/GeometryFlow/Source/GeometryFlowCore/Public/GeometryFlowTypes.h
new file mode 100644
--- /dev/null
+++ b/Engine/Plugins/Experimental/GeometryFlow/Source/GeometryFlowCore/Public/GeometryFlowTypes.h
@@ -0,0 +1,77 @@
+// Copyright Epic Games, Inc. All Rights Reserved.
+
+#pragma once
+
+#include "CoreMinimal.h"
+
+namespace UE
+{
+namespace GeometryFlow
+{
+
+// alias for Thread-Safe TSharedPtr
+template <typename T> 
+using TSafeSharedPtr = TSharedPtr<T, ESPMode::ThreadSafe>;
+
+// wrapper around MakeShared for creating a TSafeSharedPtr
+template<typename T, typename... InArgTypes>
+TSafeSharedPtr<T> MakeSafeShared(InArgTypes&&... Args)
+{
+	return MakeShared<T, ESPMode::ThreadSafe>(Forward<InArgTypes>(Args)...);
+}
+
+
+// predeclare core types
+class IData;
+class INodeInput;
+class INodeOutput;
+class FNode;
+class FGraph;
+
+
+
+enum class EGeometryFlowResult
+{
+	Ok = 0,
+	NodeDoesNotExist = 1,
+	InputDoesNotExist = 2,
+	OutputDoesNotExist = 3,
+	UnmatchedTypes = 4,
+	ConnectionDoesNotExist = 5,
+	NoMatchesFound = 6,
+	MultipleMatchingAmbiguityFound = 7
+};
+
+
+/**
+ * Every data type in a GeometryFlow graph requires a unique integer to identify the type.
+ * This enum defines various integer ranges used by the known data/node libraries.
+ */
+enum class EDataTypes
+{
+	Integer = 0,
+	Float = 1,
+	Double = 2,
+	
+	Vector3f = 20,
+	Vector3d = 21,
+
+	BaseMeshProcessingTypes = 1000,
+	BaseMeshProcessingTypesEditor = 2000,
+
+
+	UserDefinedTypes = 10000
+};
+
+
+/**
+ * This macro declares the type identifier for a GeometryFlow data type.
+ * Generally TypeID is declared in an enum, like EDataTypes above, or EMeshProcessingDataTypes in GeometryFlowMeshProcessing.
+ * Standard usage is to include at the top of the relevant class/struct definition.
+ */
+#define DECLARE_GEOMETRYFLOW_DATA_TYPE_IDENTIFIER(TypeID) static constexpr int DataTypeIdentifier = (int)TypeID;
+
+
+
+}	// end namespace GeometryFlow
+}	// end namespace UE
\ No newline at end of file
diff --git a/Engine/Plugins/Experimental/GeometryFlow/Source/GeometryFlowCore/Public/MathNodes/ArithmeticNodes.h b/Engine/Plugins/Experimental/GeometryFlow/Source/GeometryFlowCore/Public/MathNodes/ArithmeticNodes.h
new file mode 100644
--- /dev/null
+++ b/Engine/Plugins/Experimental/GeometryFlow/Source/GeometryFlowCore/Public/MathNodes/ArithmeticNodes.h
@@ -0,0 +1,27 @@
+// Copyright Epic Games, Inc. All Rights Reserved.
+
+#pragma once
+
+#include "GeometryFlowCoreNodes.h"
+#include "BaseNodes/BinaryOpLambdaNode.h"
+
+
+namespace UE
+{
+namespace GeometryFlow
+{
+
+	template<typename T, int StorageTypeIdentifier>
+	class TBinaryOpAddNode : public TBinaryOpLambdaNode<T, StorageTypeIdentifier>
+	{
+	public:
+		TBinaryOpAddNode() : TBinaryOpLambdaNode<T, StorageTypeIdentifier>([](const T& A, const T& B) { return A + B; })
+		{}
+	};
+
+	typedef TBinaryOpAddNode<float, (int)EDataTypes::Float> FAddFloatNode;
+
+
+
+}	// end namespace GeometryFlow
+}	// end namespace UE
\ No newline at end of file
diff --git a/Engine/Plugins/Experimental/GeometryFlow/Source/GeometryFlowMeshProcessing/GeometryFlowMeshProcessing.Build.cs b/Engine/Plugins/Experimental/GeometryFlow/Source/GeometryFlowMeshProcessing/GeometryFlowMeshProcessing.Build.cs
new file mode 100644
--- /dev/null
+++ b/Engine/Plugins/Experimental/GeometryFlow/Source/GeometryFlowMeshProcessing/GeometryFlowMeshProcessing.Build.cs
@@ -0,0 +1,53 @@
+// Copyright Epic Games, Inc. All Rights Reserved.
+
+using UnrealBuildTool;
+
+public class GeometryFlowMeshProcessing : ModuleRules
+{
+	public GeometryFlowMeshProcessing(ReadOnlyTargetRules Target) : base(Target)
+	{
+		PCHUsage = ModuleRules.PCHUsageMode.UseExplicitOrSharedPCHs;
+		
+		PublicIncludePaths.AddRange(
+			new string[] {
+				// ... add public include paths required here ...
+			}
+			);
+				
+		
+		PrivateIncludePaths.AddRange(
+			new string[] {
+				// ... add other private include paths required here ...
+			}
+			);
+			
+		
+		PublicDependencyModuleNames.AddRange(
+			new string[]
+			{
+				"Core",
+				"GeometricObjects",
+				"GeometryAlgorithms",
+				"DynamicMesh",
+				"GeometryFlowCore"
+			}
+            );
+			
+		
+		PrivateDependencyModuleNames.AddRange(
+			new string[]
+			{
+				//"CoreUObject",
+				//"Engine"
+			}
+			);
+		
+		
+		DynamicallyLoadedModuleNames.AddRange(
+			new string[]
+			{
+				// ... add any modules that your module loads dynamically here ...
+			}
+			);
+	}
+}
diff --git a/Engine/Plugins/Experimental/GeometryFlow/Source/GeometryFlowMeshProcessing/Private/DataTypes/MeshImageBakingData.cpp b/Engine/Plugins/Experimental/GeometryFlow/Source/GeometryFlowMeshProcessing/Private/DataTypes/MeshImageBakingData.cpp
new file mode 100644
--- /dev/null
+++ b/Engine/Plugins/Experimental/GeometryFlow/Source/GeometryFlowMeshProcessing/Private/DataTypes/MeshImageBakingData.cpp
@@ -0,0 +1,54 @@
+// Copyright Epic Games, Inc. All Rights Reserved.
+
+#include "DataTypes/MeshImageBakingData.h"
+
+using namespace UE::GeometryFlow;
+
+void FMakeMeshBakingCacheNode::Evaluate(
+	const FNamedDataMap& DatasIn,
+	FNamedDataMap& DatasOut,
+	TUniquePtr<FEvaluationInfo>& EvaluationInfo)
+{
+	if (ensure(DatasOut.Contains(OutParamCache())))
+	{
+		bool bAllInputsValid = true;
+		bool bRecomputeRequired = (IsOutputAvailable(OutParamCache()) == false);
+		TSafeSharedPtr<IData> DetailMeshArg = FindAndUpdateInputForEvaluate(InParamDetailMesh(), DatasIn, bRecomputeRequired, bAllInputsValid);
+		TSafeSharedPtr<IData> TargetMeshArg = FindAndUpdateInputForEvaluate(InParamTargetMesh(), DatasIn, bRecomputeRequired, bAllInputsValid);
+		TSafeSharedPtr<IData> SettingsArg = FindAndUpdateInputForEvaluate(InParamSettings(), DatasIn, bRecomputeRequired, bAllInputsValid);
+		if (bAllInputsValid)
+		{
+			if (bRecomputeRequired)
+			{
+				// always make a copy of settings
+				FMeshMakeBakingCacheSettings Settings;
+				SettingsArg->GetDataCopy(Settings, FMeshMakeBakingCacheSettings::DataTypeIdentifier);
+
+				// todo: support stealing these meshes if they are mutable?
+				bool bDetailIsMeshMutable = DatasIn.GetDataFlags(InParamDetailMesh()).bIsMutableData;
+				bool bTargetIsMeshMutable = DatasIn.GetDataFlags(InParamTargetMesh()).bIsMutableData;
+
+
+				TUniquePtr<FMeshBakingCache> NewCache = MakeUnique<FMeshBakingCache>();
+				NewCache->DetailMesh = DetailMeshArg->GetDataConstRef<FDynamicMesh3>((int)EMeshProcessingDataTypes::DynamicMesh);
+				NewCache->DetailSpatial = FDynamicMeshAABBTree3(&NewCache->DetailMesh, true);
+				NewCache->TargetMesh = TargetMeshArg->GetDataConstRef<FDynamicMesh3>((int)EMeshProcessingDataTypes::DynamicMesh);
+
+				NewCache->BakeCache.SetDetailMesh(&NewCache->DetailMesh, &NewCache->DetailSpatial);
+				NewCache->BakeCache.SetBakeTargetMesh(&NewCache->TargetMesh);
+
+				NewCache->BakeCache.SetDimensions(Settings.Dimensions);
+				NewCache->BakeCache.SetUVLayer(Settings.UVLayer);
+				NewCache->BakeCache.SetThickness(Settings.Thickness);
+
+				NewCache->BakeCache.ValidateCache();
+
+				TSafeSharedPtr<FMeshBakingCacheData> Result = MakeShared<FMeshBakingCacheData, ESPMode::ThreadSafe>(MoveTemp(NewCache));
+				SetOutput(OutParamCache(), Result);
+
+				EvaluationInfo->CountCompute(this);
+			}
+			DatasOut.SetData(OutParamCache(), GetOutput(OutParamCache()));
+		}
+	}
+}
\ No newline at end of file
diff --git a/Engine/Plugins/Experimental/GeometryFlow/Source/GeometryFlowMeshProcessing/Private/GeometryFlowMeshProcessingModule.cpp b/Engine/Plugins/Experimental/GeometryFlow/Source/GeometryFlowMeshProcessing/Private/GeometryFlowMeshProcessingModule.cpp
new file mode 100644
--- /dev/null
+++ b/Engine/Plugins/Experimental/GeometryFlow/Source/GeometryFlowMeshProcessing/Private/GeometryFlowMeshProcessingModule.cpp
@@ -0,0 +1,20 @@
+// Copyright Epic Games, Inc. All Rights Reserved.
+
+#include "GeometryFlowMeshProcessingModule.h"
+
+#define LOCTEXT_NAMESPACE "FGeometryFlowMeshProcessingModule"
+
+void FGeometryFlowMeshProcessingModule::StartupModule()
+{
+	// This code will execute after your module is loaded into memory; the exact timing is specified in the .uplugin file per-module
+}
+
+void FGeometryFlowMeshProcessingModule::ShutdownModule()
+{
+	// This function may be called during shutdown to clean up your module.  For modules that support dynamic reloading,
+	// we call this function before unloading the module.
+}
+
+#undef LOCTEXT_NAMESPACE
+	
+IMPLEMENT_MODULE(FGeometryFlowMeshProcessingModule, GeometryFlowMeshProcessing)
diff --git a/Engine/Plugins/Experimental/GeometryFlow/Source/GeometryFlowMeshProcessing/Private/MeshBakingNodes/BakeMeshNormalMapNode.cpp b/Engine/Plugins/Experimental/GeometryFlow/Source/GeometryFlowMeshProcessing/Private/MeshBakingNodes/BakeMeshNormalMapNode.cpp
new file mode 100644
--- /dev/null
+++ b/Engine/Plugins/Experimental/GeometryFlow/Source/GeometryFlowMeshProcessing/Private/MeshBakingNodes/BakeMeshNormalMapNode.cpp
@@ -0,0 +1,48 @@
+// Copyright Epic Games, Inc. All Rights Reserved.
+
+#include "MeshBakingNodes/BakeMeshNormalMapNode.h"
+#include "GeometryFlowNodeUtil.h"
+
+#include "Sampling/MeshNormalMapBaker.h"
+
+using namespace UE::GeometryFlow;
+
+void FBakeMeshNormalMapNode::Evaluate(
+	const FNamedDataMap& DatasIn,
+	FNamedDataMap& DatasOut,
+	TUniquePtr<FEvaluationInfo>& EvaluationInfo)
+{
+	if (ensure(DatasOut.Contains(OutParamNormalMap())))
+	{
+		bool bAllInputsValid = true;
+		bool bRecomputeRequired = (IsOutputAvailable(OutParamNormalMap()) == false);
+		TSafeSharedPtr<IData> BakeCacheArg = FindAndUpdateInputForEvaluate(InParamBakeCache(), DatasIn, bRecomputeRequired, bAllInputsValid);
+		TSafeSharedPtr<IData> TangentsArg = FindAndUpdateInputForEvaluate(InParamTangents(), DatasIn, bRecomputeRequired, bAllInputsValid);
+		TSafeSharedPtr<IData> SettingsArg = FindAndUpdateInputForEvaluate(InParamSettings(), DatasIn, bRecomputeRequired, bAllInputsValid);
+		if (bAllInputsValid)
+		{
+			if (bRecomputeRequired)
+			{
+				// always make a copy of settings
+				FBakeMeshNormalMapSettings Settings;
+				SettingsArg->GetDataCopy(Settings, FBakeMeshNormalMapSettings::DataTypeIdentifier);
+
+				const FMeshTangentsd& BaseTangents = TangentsArg->GetDataConstRef<FMeshTangentsd>((int)EMeshProcessingDataTypes::MeshTangentSet);
+				const FMeshBakingCache& BakeCacheContainer = BakeCacheArg->GetDataConstRef<FMeshBakingCache>((int)EMeshProcessingDataTypes::BakingCache);
+
+				FMeshNormalMapBaker NormalBaker;
+				NormalBaker.SetCache( &BakeCacheContainer.BakeCache );
+				NormalBaker.BaseMeshTangents = &BaseTangents;
+				NormalBaker.Bake();
+
+				FNormalMapImage ResultImage;
+				ResultImage.Image = MoveTemp(*NormalBaker.TakeResult());
+
+				SetOutput(OutParamNormalMap(), MakeMovableData<FNormalMapImage>(MoveTemp(ResultImage)));
+
+				EvaluationInfo->CountCompute(this);
+			}
+			DatasOut.SetData(OutParamNormalMap(), GetOutput(OutParamNormalMap()));
+		}
+	}
+}
\ No newline at end of file
diff --git a/Engine/Plugins/Experimental/GeometryFlow/Source/GeometryFlowMeshProcessing/Private/MeshBakingNodes/BakeMeshTextureImageNode.cpp b/Engine/Plugins/Experimental/GeometryFlow/Source/GeometryFlowMeshProcessing/Private/MeshBakingNodes/BakeMeshTextureImageNode.cpp
new file mode 100644
--- /dev/null
+++ b/Engine/Plugins/Experimental/GeometryFlow/Source/GeometryFlowMeshProcessing/Private/MeshBakingNodes/BakeMeshTextureImageNode.cpp
@@ -0,0 +1,55 @@
+// Copyright Epic Games, Inc. All Rights Reserved.
+
+#include "MeshBakingNodes/BakeMeshTextureImageNode.h"
+#include "GeometryFlowNodeUtil.h"
+
+#include "Sampling/MeshResampleImageBaker.h"
+
+using namespace UE::GeometryFlow;
+
+void FBakeMeshTextureImageNode::Evaluate(
+	const FNamedDataMap& DatasIn,
+	FNamedDataMap& DatasOut,
+	TUniquePtr<FEvaluationInfo>& EvaluationInfo)
+{
+	if (ensure(DatasOut.Contains(OutParamTextureImage())))
+	{
+		bool bAllInputsValid = true;
+		bool bRecomputeRequired = (IsOutputAvailable(OutParamTextureImage()) == false);
+		TSafeSharedPtr<IData> BakeCacheArg = FindAndUpdateInputForEvaluate(InParamBakeCache(), DatasIn, bRecomputeRequired, bAllInputsValid);
+		TSafeSharedPtr<IData> ImageArg = FindAndUpdateInputForEvaluate(InParamImage(), DatasIn, bRecomputeRequired, bAllInputsValid);
+		TSafeSharedPtr<IData> SettingsArg = FindAndUpdateInputForEvaluate(InParamSettings(), DatasIn, bRecomputeRequired, bAllInputsValid);
+		if (bAllInputsValid)
+		{
+			if (bRecomputeRequired)
+			{
+				// always make a copy of settings
+				FBakeMeshTextureImageSettings Settings;
+				SettingsArg->GetDataCopy(Settings, FBakeMeshTextureImageSettings::DataTypeIdentifier);
+
+				const FMeshBakingCache& BakeCacheContainer = BakeCacheArg->GetDataConstRef<FMeshBakingCache>((int)EMeshProcessingDataTypes::BakingCache);
+				const FDynamicMesh3& DetailMesh = BakeCacheContainer.DetailMesh;
+				check(DetailMesh.HasAttributes() && DetailMesh.Attributes()->NumUVLayers() > Settings.DetailUVLayer);
+				const FDynamicMeshUVOverlay* DetailUVOverlay = DetailMesh.Attributes()->GetUVLayer(Settings.DetailUVLayer);
+
+				const FTextureImage& SourceImage = ImageArg->GetDataConstRef<FTextureImage>(FTextureImage::DataTypeIdentifier);
+
+				FMeshResampleImageBaker Baker;
+				Baker.SetCache(&BakeCacheContainer.BakeCache);
+				Baker.DetailUVOverlay = DetailUVOverlay;
+				Baker.SampleFunction = [&](FVector2d UVCoord) {
+					return SourceImage.Image.BilinearSampleUV<float>(UVCoord, FVector4f(0, 0, 0, 1));
+				};
+				Baker.Bake();
+
+				FTextureImage ResultImage;
+				ResultImage.Image = MoveTemp(*Baker.TakeResult());
+
+				SetOutput(OutParamTextureImage(), MakeMovableData<FTextureImage>(MoveTemp(ResultImage)) );
+
+				EvaluationInfo->CountCompute(this);
+			}
+			DatasOut.SetData(OutParamTextureImage(), GetOutput(OutParamTextureImage()));
+		}
+	}
+}
\ No newline at end of file
diff --git a/Engine/Plugins/Experimental/GeometryFlow/Source/GeometryFlowMeshProcessing/Private/MeshProcessingNodes/MeshRecalculateUVsNode.cpp b/Engine/Plugins/Experimental/GeometryFlow/Source/GeometryFlowMeshProcessing/Private/MeshProcessingNodes/MeshRecalculateUVsNode.cpp
new file mode 100644
--- /dev/null
+++ b/Engine/Plugins/Experimental/GeometryFlow/Source/GeometryFlowMeshProcessing/Private/MeshProcessingNodes/MeshRecalculateUVsNode.cpp
@@ -0,0 +1,61 @@
+// Copyright Epic Games, Inc. All Rights Reserved.
+
+#include "MeshProcessingNodes/MeshRecalculateUVsNode.h"
+
+#include "Selections/MeshConnectedComponents.h"
+#include "Parameterization/DynamicMeshUVEditor.h"
+#include "Parameterization/MeshUVTransforms.h"
+
+
+
+using namespace UE::GeometryFlow;
+
+
+void FMeshRecalculateUVsNode::RecalculateUVsOnMesh(FDynamicMesh3& EditMesh, const FMeshRecalculateUVsSettings& Settings)
+{
+	FDynamicMeshUVEditor UVEditor(&EditMesh, Settings.UVLayer, true);
+	FDynamicMeshUVOverlay* UVOverlay = UVEditor.GetOverlay();
+
+	FMeshConnectedComponents UVComponents(&EditMesh);
+	UVComponents.FindConnectedTriangles([&](int32 Triangle0, int32 Triangle1) {
+		return UVOverlay->AreTrianglesConnected(Triangle0, Triangle1);
+	});
+
+	for (int32 k = 0; k < UVComponents.Num(); ++k)
+	{
+		FMeshConnectedComponents::FComponent& Component = UVComponents.GetComponent(k);
+		if (!ensure(Component.Indices.Num() > 0))
+		{
+			continue;		// how??
+		}
+
+		FUVEditResult EditResult;
+
+		ERecalculateUVsUnwrapType UseUnwrapType = Settings.UnwrapType;
+		if (UseUnwrapType == ERecalculateUVsUnwrapType::Auto)
+		{
+			UseUnwrapType = ERecalculateUVsUnwrapType::Conformal;
+		}
+
+		bool bComplete = false;
+		if (UseUnwrapType == ERecalculateUVsUnwrapType::Conformal)
+		{
+			bComplete = UVEditor.SetTriangleUVsFromFreeBoundaryConformal(Component.Indices, &EditResult);
+			// if this fails we will fall back to ExpMap
+		}
+
+		if (!bComplete)
+		{
+			EditResult = FUVEditResult();
+			if (UVEditor.SetTriangleUVsFromExpMap(Component.Indices) == false)
+			{
+				// if we somehow failed at conformal, fallback to trivial planar projection
+				UVEditor.SetTriangleUVsFromProjection(Component.Indices, FFrame3d());
+			}
+		}
+
+	}
+}
+
+
+
diff --git a/Engine/Plugins/Experimental/GeometryFlow/Source/GeometryFlowMeshProcessing/Private/MeshProcessingNodes/MeshRepackUVsNode.cpp b/Engine/Plugins/Experimental/GeometryFlow/Source/GeometryFlowMeshProcessing/Private/MeshProcessingNodes/MeshRepackUVsNode.cpp
new file mode 100644
--- /dev/null
+++ b/Engine/Plugins/Experimental/GeometryFlow/Source/GeometryFlowMeshProcessing/Private/MeshProcessingNodes/MeshRepackUVsNode.cpp
@@ -0,0 +1,40 @@
+// Copyright Epic Games, Inc. All Rights Reserved.
+
+#include "MeshProcessingNodes/MeshRepackUVsNode.h"
+
+#include "Parameterization/MeshUVPacking.h"
+
+using namespace UE::GeometryFlow;
+
+
+void FMeshRepackUVsNode::RepackUVsForMesh(FDynamicMesh3& EditMesh, const FMeshRepackUVsSettings& Settings)
+{
+	FDynamicMeshUVOverlay* UVLayer = EditMesh.Attributes()->GetUVLayer(Settings.UVLayer);
+
+	UVLayer->SplitBowties();
+
+	FDynamicMeshUVPacker Packer(UVLayer);
+	Packer.TextureResolution = Settings.TextureResolution;
+	Packer.GutterSize = Settings.GutterSize;
+	Packer.bAllowFlips = Settings.bAllowFlips;
+
+	bool bOK = Packer.StandardPack();
+	ensure(bOK);
+
+
+	if (Settings.UVScale != FVector2f::One() || Settings.UVTranslation != FVector2f::Zero())
+	{
+		for (int ElementID : UVLayer->ElementIndicesItr())
+		{
+			FVector2f UV = UVLayer->GetElement(ElementID);
+			UV.X *= Settings.UVScale.X;
+			UV.Y *= Settings.UVScale.Y;
+			UV += Settings.UVTranslation;
+			UVLayer->SetElement(ElementID, UV);
+		}
+	}
+
+}
+
+
+
diff --git a/Engine/Plugins/Experimental/GeometryFlow/Source/GeometryFlowMeshProcessing/Private/PhysicsNodes/GenerateConvexHullsCollisionNode.cpp b/Engine/Plugins/Experimental/GeometryFlow/Source/GeometryFlowMeshProcessing/Private/PhysicsNodes/GenerateConvexHullsCollisionNode.cpp
new file mode 100644
--- /dev/null
+++ b/Engine/Plugins/Experimental/GeometryFlow/Source/GeometryFlowMeshProcessing/Private/PhysicsNodes/GenerateConvexHullsCollisionNode.cpp
@@ -0,0 +1,66 @@
+// Copyright Epic Games, Inc. All Rights Reserved.
+
+#include "PhysicsNodes/GenerateConvexHullsCollisionNode.h"
+
+#include "MeshIndexUtil.h"
+#include "Operations/MeshConvexHull.h"
+
+#include "Async/ParallelFor.h"
+
+
+using namespace UE::GeometryFlow;
+
+void FGenerateConvexHullsCollisionNode::Evaluate(
+	const FNamedDataMap& DatasIn,
+	FNamedDataMap& DatasOut,
+	TUniquePtr<FEvaluationInfo>& EvaluationInfo)
+{
+	if (ensure(DatasOut.Contains(OutParamGeometry())))
+	{
+		bool bAllInputsValid = true;
+		bool bRecomputeRequired = (IsOutputAvailable(OutParamGeometry()) == false);
+		TSafeSharedPtr<IData> MeshArg = FindAndUpdateInputForEvaluate(InParamMesh(), DatasIn, bRecomputeRequired, bAllInputsValid);
+		TSafeSharedPtr<IData> TriSetsArg = FindAndUpdateInputForEvaluate(InParamIndexSets(), DatasIn, bRecomputeRequired, bAllInputsValid);
+		TSafeSharedPtr<IData> SettingsArg = FindAndUpdateInputForEvaluate(InParamSettings(), DatasIn, bRecomputeRequired, bAllInputsValid);
+		if (bAllInputsValid)
+		{
+			if (bRecomputeRequired)
+			{
+				// always make a copy of settings
+				FGenerateConvexHullsCollisionSettings Settings;
+				SettingsArg->GetDataCopy(Settings, FGenerateConvexHullsCollisionSettings::DataTypeIdentifier);
+
+				const FDynamicMesh3& Mesh = MeshArg->GetDataConstRef<FDynamicMesh3>((int)EMeshProcessingDataTypes::DynamicMesh);
+				const FIndexSets& IndexData = TriSetsArg->GetDataConstRef<FIndexSets>((int)EMeshProcessingDataTypes::IndexSets);
+
+				int32 NumShapes = IndexData.IndexSets.Num();
+				TArray<FConvexShape3d> Convexes;
+				TArray<bool> ConvexOk;
+				Convexes.SetNum(NumShapes);
+				ConvexOk.SetNum(NumShapes);
+				ParallelFor(NumShapes, [&](int32 k)
+				{
+					ConvexOk[k] = false;
+					FMeshConvexHull Hull(&Mesh);
+					MeshIndexUtil::TriangleToVertexIDs(&Mesh, IndexData.IndexSets[k], Hull.VertexSet);
+					Hull.bPostSimplify = Settings.SimplifyToTriangleCount > 0;
+					Hull.MaxTargetFaceCount = FMath::Max(1, Settings.SimplifyToTriangleCount);
+					if (Hull.Compute())
+					{
+						FConvexShape3d NewConvex;
+						NewConvex.Mesh = MoveTemp(Hull.ConvexHull);
+						Convexes[k] = MoveTemp(NewConvex);
+						ConvexOk[k] = true;
+					}
+				});
+
+				FCollisionGeometry Result;
+				Result.Geometry.Convexes = MoveTemp(Convexes);
+
+				SetOutput(OutParamGeometry(), MakeMovableData<FCollisionGeometry>(MoveTemp(Result)));
+				EvaluationInfo->CountCompute(this);
+			}
+			DatasOut.SetData(OutParamGeometry(), GetOutput(OutParamGeometry()));
+		}
+	}
+}
\ No newline at end of file
diff --git a/Engine/Plugins/Experimental/GeometryFlow/Source/GeometryFlowMeshProcessing/Public/DataTypes/CollisionGeometryData.h b/Engine/Plugins/Experimental/GeometryFlow/Source/GeometryFlowMeshProcessing/Public/DataTypes/CollisionGeometryData.h
new file mode 100644
--- /dev/null
+++ b/Engine/Plugins/Experimental/GeometryFlow/Source/GeometryFlowMeshProcessing/Public/DataTypes/CollisionGeometryData.h
@@ -0,0 +1,36 @@
+// Copyright Epic Games, Inc. All Rights Reserved.
+
+#pragma once
+
+#include "GeometryFlowCoreNodes.h"
+#include "GeometryFlowMovableData.h"
+#include "BaseNodes/TransferNode.h"
+#include "MeshProcessingNodes/MeshProcessingDataTypes.h"
+
+#include "ShapeApproximation/SimpleShapeSet3.h"
+
+
+namespace UE
+{
+namespace GeometryFlow
+{
+
+
+struct FCollisionGeometry
+{
+	DECLARE_GEOMETRYFLOW_DATA_TYPE_IDENTIFIER(EMeshProcessingDataTypes::CollisionGeometry);
+
+	FSimpleShapeSet3d Geometry;
+};
+
+
+// declares FDataCollisionGeometry, FCollisionGeometryInput, FCollisionGeometryOutput, FCollisionGeometrySourceNode
+GEOMETRYFLOW_DECLARE_BASIC_TYPES(CollisionGeometry, FCollisionGeometry, (int)EMeshProcessingDataTypes::CollisionGeometry);
+
+
+typedef TTransferNode<FCollisionGeometry, (int)EMeshProcessingDataTypes::CollisionGeometry> FCollisionGeometryTransferNode;
+
+
+
+}	// end namespace GeometryFlow
+}	// end namespace UE
\ No newline at end of file
diff --git a/Engine/Plugins/Experimental/GeometryFlow/Source/GeometryFlowMeshProcessing/Public/DataTypes/DynamicMeshData.h b/Engine/Plugins/Experimental/GeometryFlow/Source/GeometryFlowMeshProcessing/Public/DataTypes/DynamicMeshData.h
new file mode 100644
--- /dev/null
+++ b/Engine/Plugins/Experimental/GeometryFlow/Source/GeometryFlowMeshProcessing/Public/DataTypes/DynamicMeshData.h
@@ -0,0 +1,23 @@
+// Copyright Epic Games, Inc. All Rights Reserved.
+
+#pragma once
+
+#include "GeometryFlowCoreNodes.h"
+#include "BaseNodes/TransferNode.h"
+#include "MeshProcessingNodes/MeshProcessingDataTypes.h"
+#include "DynamicMesh3.h"
+
+namespace UE
+{
+namespace GeometryFlow
+{
+
+// declares FDataDynamicMesh, FDynamicMeshInput, FDynamicMeshOutput, FDynamicMeshSourceNode
+GEOMETRYFLOW_DECLARE_BASIC_TYPES(DynamicMesh, FDynamicMesh3, (int)EMeshProcessingDataTypes::DynamicMesh)
+
+
+typedef TTransferNode<FDynamicMesh3, (int)EMeshProcessingDataTypes::DynamicMesh> FDynamicMeshTransferNode;
+
+
+}	// end namespace GeometryFlow
+}	// end namespace UE
\ No newline at end of file
diff --git a/Engine/Plugins/Experimental/GeometryFlow/Source/GeometryFlowMeshProcessing/Public/DataTypes/IndexSetsData.h b/Engine/Plugins/Experimental/GeometryFlow/Source/GeometryFlowMeshProcessing/Public/DataTypes/IndexSetsData.h
new file mode 100644
--- /dev/null
+++ b/Engine/Plugins/Experimental/GeometryFlow/Source/GeometryFlowMeshProcessing/Public/DataTypes/IndexSetsData.h
@@ -0,0 +1,61 @@
+// Copyright Epic Games, Inc. All Rights Reserved.
+
+#pragma once
+
+#include "GeometryFlowCoreNodes.h"
+#include "GeometryFlowMovableData.h"
+#include "MeshProcessingNodes/MeshProcessingDataTypes.h"
+
+
+namespace UE
+{
+namespace GeometryFlow
+{
+
+
+struct FIndexSets
+{
+	DECLARE_GEOMETRYFLOW_DATA_TYPE_IDENTIFIER(EMeshProcessingDataTypes::IndexSets);
+
+	TArray<TArray<int32>> IndexSets;
+
+
+	//
+	// utility functions
+	//
+	void AppendSet()
+	{
+		TArray<int32> Values;
+		IndexSets.Add(Values);
+	}
+
+	void AppendSet(const TArray<int32>& SetValues)
+	{
+		IndexSets.Add(SetValues);
+	}
+
+	int32 NumSets() const { return IndexSets.Num(); }
+
+	template<typename ListType>
+	void GetAllValues(ListType& ValuesOut) const
+	{
+		for (const TArray<int32>& Set : IndexSets)
+		{
+			for (int32 Value : Set)
+			{
+				ValuesOut.Add(Value);
+			}
+		}
+	}
+};
+
+
+// declares FDataIndexSets, FIndexSetsInput, FIndexSetsOutput, FIndexSetsSourceNode
+GEOMETRYFLOW_DECLARE_BASIC_TYPES(IndexSets, FIndexSets, (int)EMeshProcessingDataTypes::IndexSets)
+
+
+
+
+
+}	// end namespace GeometryFlow
+}	// end namespace UE
\ No newline at end of file
diff --git a/Engine/Plugins/Experimental/GeometryFlow/Source/GeometryFlowMeshProcessing/Public/DataTypes/MeshImageBakingData.h b/Engine/Plugins/Experimental/GeometryFlow/Source/GeometryFlowMeshProcessing/Public/DataTypes/MeshImageBakingData.h
new file mode 100644
--- /dev/null
+++ b/Engine/Plugins/Experimental/GeometryFlow/Source/GeometryFlowMeshProcessing/Public/DataTypes/MeshImageBakingData.h
@@ -0,0 +1,88 @@
+// Copyright Epic Games, Inc. All Rights Reserved.
+
+#pragma once
+
+#include "GeometryFlowCoreNodes.h"
+#include "GeometryFlowImmutableData.h"
+#include "DataTypes/DynamicMeshData.h"
+
+#include "DynamicMesh3.h"
+#include "DynamicMeshAABBTree3.h"
+#include "Sampling/MeshImageBakingCache.h"
+
+namespace UE
+{
+namespace GeometryFlow
+{
+
+
+
+struct GEOMETRYFLOWMESHPROCESSING_API FMeshBakingCache
+{
+	DECLARE_GEOMETRYFLOW_DATA_TYPE_IDENTIFIER(EMeshProcessingDataTypes::BakingCache);
+
+	// have to make copies of these - boo!
+	FDynamicMesh3 DetailMesh;
+	FDynamicMeshAABBTree3 DetailSpatial;
+	FDynamicMesh3 TargetMesh;
+
+	FMeshImageBakingCache BakeCache;
+};
+
+
+typedef TImmutableData<FMeshBakingCache, FMeshBakingCache::DataTypeIdentifier> FMeshBakingCacheData;
+
+
+
+
+struct GEOMETRYFLOWMESHPROCESSING_API FMeshMakeBakingCacheSettings
+{
+	DECLARE_GEOMETRYFLOW_DATA_TYPE_IDENTIFIER(EMeshProcessingDataTypes::MakeBakingCacheSettings);
+
+	FImageDimensions Dimensions;
+	int UVLayer = 0;
+	float Thickness = 0.1;
+
+};
+GEOMETRYFLOW_DECLARE_SETTINGS_TYPES(FMeshMakeBakingCacheSettings, MeshMakeBakingCache);
+
+
+
+
+class GEOMETRYFLOWMESHPROCESSING_API FMakeMeshBakingCacheNode : public FNode
+{
+protected:
+	using SettingsDataType = TMovableData<FMeshMakeBakingCacheSettings, FMeshMakeBakingCacheSettings::DataTypeIdentifier>;
+
+public:
+	static const FString InParamDetailMesh() { return TEXT("DetailMesh"); }
+	static const FString InParamTargetMesh() { return TEXT("TargetMesh"); }
+	static const FString InParamSettings() { return TEXT("Settings"); }
+
+	static const FString OutParamCache() { return TEXT("BakeCache"); }
+
+public:
+	FMakeMeshBakingCacheNode()
+	{
+		AddInput(InParamDetailMesh(), MakeUnique<FDynamicMeshInput>());
+		AddInput(InParamTargetMesh(), MakeUnique<FDynamicMeshInput>());
+		AddInput(InParamSettings(), MakeUnique<TBasicNodeInput<FMeshMakeBakingCacheSettings, FMeshMakeBakingCacheSettings::DataTypeIdentifier>>());
+
+		AddOutput(OutParamCache(), MakeUnique<TImmutableNodeOutput<FMeshBakingCache, FMeshBakingCache::DataTypeIdentifier>>() );
+	}
+
+	virtual void Evaluate(
+		const FNamedDataMap& DatasIn,
+		FNamedDataMap& DatasOut,
+		TUniquePtr<FEvaluationInfo>& EvaluationInfo) override;
+};
+
+
+
+
+
+
+
+
+}	// end namespace GeometryFlow
+}	// end 
\ No newline at end of file
diff --git a/Engine/Plugins/Experimental/GeometryFlow/Source/GeometryFlowMeshProcessing/Public/DataTypes/NormalMapData.h b/Engine/Plugins/Experimental/GeometryFlow/Source/GeometryFlowMeshProcessing/Public/DataTypes/NormalMapData.h
new file mode 100644
--- /dev/null
+++ b/Engine/Plugins/Experimental/GeometryFlow/Source/GeometryFlowMeshProcessing/Public/DataTypes/NormalMapData.h
@@ -0,0 +1,34 @@
+// Copyright Epic Games, Inc. All Rights Reserved.
+
+#pragma once
+
+#include "GeometryFlowCoreNodes.h"
+#include "GeometryFlowMovableData.h"
+#include "MeshProcessingNodes/MeshProcessingDataTypes.h"
+
+#include "Image/ImageBuilder.h"
+
+
+namespace UE
+{
+namespace GeometryFlow
+{
+
+
+struct FNormalMapImage
+{
+	DECLARE_GEOMETRYFLOW_DATA_TYPE_IDENTIFIER(EMeshProcessingDataTypes::NormalMapImage);
+
+	TImageBuilder<FVector3f> Image;
+};
+
+
+// declares FDataNormalMapImage, FNormalMapImageInput, FNormalMapImageOutput, FNormalMapImageSourceNode
+GEOMETRYFLOW_DECLARE_BASIC_TYPES(NormalMapImage, FNormalMapImage, (int)EMeshProcessingDataTypes::NormalMapImage)
+
+
+
+
+
+}	// end namespace GeometryFlow
+}	// end namespace UE
\ No newline at end of file
diff --git a/Engine/Plugins/Experimental/GeometryFlow/Source/GeometryFlowMeshProcessing/Public/DataTypes/TextureImageData.h b/Engine/Plugins/Experimental/GeometryFlow/Source/GeometryFlowMeshProcessing/Public/DataTypes/TextureImageData.h
new file mode 100644
--- /dev/null
+++ b/Engine/Plugins/Experimental/GeometryFlow/Source/GeometryFlowMeshProcessing/Public/DataTypes/TextureImageData.h
@@ -0,0 +1,30 @@
+// Copyright Epic Games, Inc. All Rights Reserved.
+
+#pragma once
+
+#include "GeometryFlowCoreNodes.h"
+#include "GeometryFlowImmutableData.h"
+#include "MeshProcessingNodes/MeshProcessingDataTypes.h"
+
+#include "Image/ImageBuilder.h"
+
+namespace UE
+{
+namespace GeometryFlow
+{
+
+
+struct FTextureImage
+{
+	DECLARE_GEOMETRYFLOW_DATA_TYPE_IDENTIFIER(EMeshProcessingDataTypes::TextureImage);
+
+	TImageBuilder<FVector4f> Image;
+};
+
+// declares FDataTextureImage, FTextureImageInput, FTextureImageOutput, FTextureImageSourceNode
+GEOMETRYFLOW_DECLARE_BASIC_TYPES(TextureImage, FTextureImage, (int)EMeshProcessingDataTypes::TextureImage)
+
+
+
+}	// end namespace GeometryFlow
+}	// end namespace UE
\ No newline at end of file
diff --git a/Engine/Plugins/Experimental/GeometryFlow/Source/GeometryFlowMeshProcessing/Public/GeometryFlowMeshProcessingModule.h b/Engine/Plugins/Experimental/GeometryFlow/Source/GeometryFlowMeshProcessing/Public/GeometryFlowMeshProcessingModule.h
new file mode 100644
--- /dev/null
+++ b/Engine/Plugins/Experimental/GeometryFlow/Source/GeometryFlowMeshProcessing/Public/GeometryFlowMeshProcessingModule.h
@@ -0,0 +1,15 @@
+// Copyright Epic Games, Inc. All Rights Reserved.
+
+#pragma once
+
+#include "CoreMinimal.h"
+#include "Modules/ModuleManager.h"
+
+class FGeometryFlowMeshProcessingModule : public IModuleInterface
+{
+public:
+
+	/** IModuleInterface implementation */
+	virtual void StartupModule() override;
+	virtual void ShutdownModule() override;
+};
diff --git a/Engine/Plugins/Experimental/GeometryFlow/Source/GeometryFlowMeshProcessing/Public/MeshBakingNodes/BakeMeshNormalMapNode.h b/Engine/Plugins/Experimental/GeometryFlow/Source/GeometryFlowMeshProcessing/Public/MeshBakingNodes/BakeMeshNormalMapNode.h
new file mode 100644
--- /dev/null
+++ b/Engine/Plugins/Experimental/GeometryFlow/Source/GeometryFlowMeshProcessing/Public/MeshBakingNodes/BakeMeshNormalMapNode.h
@@ -0,0 +1,62 @@
+// Copyright Epic Games, Inc. All Rights Reserved.
+
+#pragma once
+
+#include "MeshProcessingNodes/MeshProcessingBaseNodes.h"
+#include "MeshProcessingNodes/MeshProcessingDataTypes.h"
+#include "MeshProcessingNodes/MeshTangentsNodes.h"
+#include "DataTypes/MeshImageBakingData.h"
+#include "DataTypes/NormalMapData.h"
+
+
+namespace UE
+{
+namespace GeometryFlow
+{
+
+
+struct GEOMETRYFLOWMESHPROCESSING_API FBakeMeshNormalMapSettings
+{
+	DECLARE_GEOMETRYFLOW_DATA_TYPE_IDENTIFIER(EMeshProcessingDataTypes::BakeNormalMapSettings);
+
+	double MaxDistance = 0.0;
+};
+GEOMETRYFLOW_DECLARE_SETTINGS_TYPES(FBakeMeshNormalMapSettings, BakeMeshNormalMap);
+
+
+
+class GEOMETRYFLOWMESHPROCESSING_API FBakeMeshNormalMapNode : public FNode
+{
+protected:
+	using SettingsDataType = TMovableData<FBakeMeshNormalMapSettings, FBakeMeshNormalMapSettings::DataTypeIdentifier>;
+
+public:
+	static const FString InParamBakeCache() { return TEXT("BakeCache"); }
+	static const FString InParamTangents() { return TEXT("Tangents"); }
+	static const FString InParamSettings() { return TEXT("Settings"); }
+
+	static const FString OutParamNormalMap() { return TEXT("NormalMap"); }
+
+public:
+	FBakeMeshNormalMapNode()
+	{
+		AddInput(InParamBakeCache(), MakeUnique<TImmutableNodeInput<FMeshBakingCache, FMeshBakingCache::DataTypeIdentifier>>());
+		AddInput(InParamTangents(), MakeUnique<TBasicNodeInput<FMeshTangentsd, (int)EMeshProcessingDataTypes::MeshTangentSet>>());
+		AddInput(InParamSettings(), MakeUnique<TBasicNodeInput<FBakeMeshNormalMapSettings, FBakeMeshNormalMapSettings::DataTypeIdentifier>>());
+
+		AddOutput(OutParamNormalMap(), MakeUnique<TBasicNodeOutput<FNormalMapImage, FNormalMapImage::DataTypeIdentifier>>());
+	}
+
+	virtual void Evaluate(
+		const FNamedDataMap& DatasIn,
+		FNamedDataMap& DatasOut,
+		TUniquePtr<FEvaluationInfo>& EvaluationInfo) override;
+};
+
+
+
+
+
+
+}	// end namespace GeometryFlow
+}	// end namespace UE
\ No newline at end of file
diff --git a/Engine/Plugins/Experimental/GeometryFlow/Source/GeometryFlowMeshProcessing/Public/MeshBakingNodes/BakeMeshTextureImageNode.h b/Engine/Plugins/Experimental/GeometryFlow/Source/GeometryFlowMeshProcessing/Public/MeshBakingNodes/BakeMeshTextureImageNode.h
new file mode 100644
--- /dev/null
+++ b/Engine/Plugins/Experimental/GeometryFlow/Source/GeometryFlowMeshProcessing/Public/MeshBakingNodes/BakeMeshTextureImageNode.h
@@ -0,0 +1,64 @@
+// Copyright Epic Games, Inc. All Rights Reserved.
+
+#pragma once
+
+#include "MeshProcessingNodes/MeshProcessingBaseNodes.h"
+#include "MeshProcessingNodes/MeshProcessingDataTypes.h"
+#include "MeshProcessingNodes/MeshTangentsNodes.h"
+#include "DataTypes/MeshImageBakingData.h"
+#include "DataTypes/TextureImageData.h"
+
+
+namespace UE
+{
+namespace GeometryFlow
+{
+
+
+struct GEOMETRYFLOWMESHPROCESSING_API FBakeMeshTextureImageSettings
+{
+	DECLARE_GEOMETRYFLOW_DATA_TYPE_IDENTIFIER(EMeshProcessingDataTypes::BakeTextureImageSettings);
+
+	int32 DetailUVLayer = 0;
+
+	double MaxDistance = 0.0;
+};
+GEOMETRYFLOW_DECLARE_SETTINGS_TYPES(FBakeMeshTextureImageSettings, BakeMeshTextureImage);
+
+
+
+class GEOMETRYFLOWMESHPROCESSING_API FBakeMeshTextureImageNode : public FNode
+{
+protected:
+	using SettingsDataType = TMovableData<FBakeMeshTextureImageSettings, FBakeMeshTextureImageSettings::DataTypeIdentifier>;
+
+public:
+	static const FString InParamBakeCache() { return TEXT("BakeCache"); }
+	static const FString InParamImage() { return TEXT("TextureImage"); }
+	static const FString InParamSettings() { return TEXT("Settings"); }
+
+	static const FString OutParamTextureImage() { return TEXT("TextureImage"); }
+
+public:
+	FBakeMeshTextureImageNode()
+	{
+		AddInput(InParamBakeCache(), MakeUnique<TImmutableNodeInput<FMeshBakingCache, FMeshBakingCache::DataTypeIdentifier>>());
+		AddInput(InParamImage(), MakeUnique<TBasicNodeInput<FTextureImage, FTextureImage::DataTypeIdentifier>>());
+		AddInput(InParamSettings(), MakeUnique<TBasicNodeInput<FBakeMeshTextureImageSettings, FBakeMeshTextureImageSettings::DataTypeIdentifier>>());
+
+		AddOutput(OutParamTextureImage(), MakeUnique<TBasicNodeOutput<FTextureImage, FTextureImage::DataTypeIdentifier>>());
+	}
+
+	virtual void Evaluate(
+		const FNamedDataMap& DatasIn,
+		FNamedDataMap& DatasOut,
+		TUniquePtr<FEvaluationInfo>& EvaluationInfo) override;
+};
+
+
+
+
+
+
+}	// end namespace GeometryFlow
+}	// end namespace UE
\ No newline at end of file
diff --git a/Engine/Plugins/Experimental/GeometryFlow/Source/GeometryFlowMeshProcessing/Public/MeshDecompositionNodes/MakeTriangleSetsNode.h b/Engine/Plugins/Experimental/GeometryFlow/Source/GeometryFlowMeshProcessing/Public/MeshDecompositionNodes/MakeTriangleSetsNode.h
new file mode 100644
--- /dev/null
+++ b/Engine/Plugins/Experimental/GeometryFlow/Source/GeometryFlowMeshProcessing/Public/MeshDecompositionNodes/MakeTriangleSetsNode.h
@@ -0,0 +1,156 @@
+// Copyright Epic Games, Inc. All Rights Reserved.
+
+#pragma once
+
+#include "GeometryFlowNodeUtil.h"
+#include "DataTypes/DynamicMeshData.h"
+#include "DataTypes/IndexSetsData.h"
+
+
+namespace UE
+{
+namespace GeometryFlow
+{
+
+
+class FMakeTriangleSetsFromMeshNode : public FNode
+{
+public:
+	static const FString InParam() { return TEXT("Mesh"); }
+
+	static const FString OutParamIndexSets() { return TEXT("IndexSets"); }
+
+public:
+	FMakeTriangleSetsFromMeshNode()
+	{
+		AddInput(InParam(), MakeUnique<FDynamicMeshInput>());
+		
+		AddOutput(OutParamIndexSets(), MakeBasicOutput<FIndexSets>());
+	}
+
+	virtual void Evaluate(
+		const FNamedDataMap& DatasIn,
+		FNamedDataMap& DatasOut,
+		TUniquePtr<FEvaluationInfo>& EvaluationInfo) override
+	{
+		if (ensure(DatasOut.Contains(OutParamIndexSets())))
+		{
+			bool bAllInputsValid = true;
+			bool bRecomputeRequired = (IsOutputAvailable(OutParamIndexSets()) == false);
+			TSafeSharedPtr<IData> MeshArg = FindAndUpdateInputForEvaluate(InParam(), DatasIn, bRecomputeRequired, bAllInputsValid);
+			if (bAllInputsValid)
+			{
+				if (bRecomputeRequired)
+				{
+					const FDynamicMesh3& Mesh = MeshArg->GetDataConstRef<FDynamicMesh3>((int)EMeshProcessingDataTypes::DynamicMesh);
+
+					FIndexSets NewSets;
+					ComputeIndexSets(DatasIn, Mesh, NewSets);
+
+					SetOutput(OutParamIndexSets(), MakeMovableData<FIndexSets>(MoveTemp(NewSets)));
+					EvaluationInfo->CountCompute(this);
+				}
+				DatasOut.SetData(OutParamIndexSets(), GetOutput(OutParamIndexSets()));
+			}
+		}
+	}
+
+
+	virtual void CheckAdditionalInputs(const FNamedDataMap& DatasIn, bool& bRecomputeRequired, bool& bAllInputsValid)
+	{
+		// none
+	}
+
+
+	virtual void ComputeIndexSets(const FNamedDataMap& DatasIn, const FDynamicMesh3& Mesh, FIndexSets& SetsOut)
+	{
+		SetsOut.IndexSets.SetNum(1);
+		SetsOut.IndexSets[0].Reserve(Mesh.TriangleCount());
+		for (int32 tid : Mesh.TriangleIndicesItr())
+		{
+			SetsOut.IndexSets[0].Add(tid);
+		}
+	}
+
+};
+
+
+
+
+class FMakeTriangleSetsFromGroupsNode : public FMakeTriangleSetsFromMeshNode
+{
+public:
+	static const FString InParamIgnoreGroups() { return TEXT("IgnoreGroups"); }
+
+
+public:
+	FMakeTriangleSetsFromGroupsNode() : FMakeTriangleSetsFromMeshNode()
+	{
+		AddInput(InParamIgnoreGroups(), MakeBasicInput<FIndexSets>());
+	}
+
+	virtual void CheckAdditionalInputs(const FNamedDataMap& DatasIn, bool& bRecomputeRequired, bool& bAllInputsValid) override
+	{
+		FindAndUpdateInputForEvaluate(InParamIgnoreGroups(), DatasIn, bRecomputeRequired, bAllInputsValid);
+	}
+
+
+public:
+	
+	virtual void ComputeIndexSets(const FNamedDataMap& DatasIn, const FDynamicMesh3& Mesh, FIndexSets& SetsOut) override
+	{
+		TSafeSharedPtr<IData> IgnoreGroupsArg = DatasIn.FindData(InParamIgnoreGroups());
+		const FIndexSets& IgnoreGroupsSets = IgnoreGroupsArg->GetDataConstRef<FIndexSets>(FIndexSets::DataTypeIdentifier);
+		TSet<int32> IgnoreGroups;
+		IgnoreGroupsSets.GetAllValues(IgnoreGroups);
+
+		TMap<int32, int32> GroupsMap;
+		TArray<int32> GroupCounts;
+		int32 NumGroups = 0;
+		for (int32 tid : Mesh.TriangleIndicesItr())
+		{
+			int32 GroupID = Mesh.GetTriangleGroup(tid);
+			if (IgnoreGroups.Contains(GroupID))
+			{
+				continue;
+			}
+
+			int32* FoundIndex = GroupsMap.Find(GroupID);
+			if (FoundIndex == nullptr)
+			{
+				int32 Index = NumGroups++;
+				GroupsMap.Add(GroupID, Index);
+				GroupCounts.Add(1);
+			}
+			else
+			{
+				GroupCounts[*FoundIndex]++;
+			}
+		}
+
+		SetsOut.IndexSets.SetNum(NumGroups);
+		for (int32 k = 0; k < NumGroups; ++k)
+		{
+			SetsOut.IndexSets[k].Reserve(GroupCounts[k]);
+		}
+
+		for (int32 tid : Mesh.TriangleIndicesItr())
+		{
+			int32 GroupID = Mesh.GetTriangleGroup(tid);
+			if (IgnoreGroups.Contains(GroupID))
+			{
+				continue;
+			}
+
+			int32* FoundIndex = GroupsMap.Find(GroupID);
+			SetsOut.IndexSets[*FoundIndex].Add(tid);
+		}
+	}
+
+};
+
+
+
+
+}	// end namespace GeometryFlow
+}	// end namespace UE
\ No newline at end of file
diff --git a/Engine/Plugins/Experimental/GeometryFlow/Source/GeometryFlowMeshProcessing/Public/MeshProcessingNodes/MeshDeleteTrianglesNode.h b/Engine/Plugins/Experimental/GeometryFlow/Source/GeometryFlowMeshProcessing/Public/MeshProcessingNodes/MeshDeleteTrianglesNode.h
new file mode 100644
--- /dev/null
+++ b/Engine/Plugins/Experimental/GeometryFlow/Source/GeometryFlowMeshProcessing/Public/MeshProcessingNodes/MeshDeleteTrianglesNode.h
@@ -0,0 +1,72 @@
+// Copyright Epic Games, Inc. All Rights Reserved.
+
+#pragma once
+
+#include "MeshProcessingNodes/MeshProcessingBaseNodes.h"
+#include "MeshProcessingNodes/MeshProcessingDataTypes.h"
+#include "DataTypes/IndexSetsData.h"
+#include "DynamicMeshEditor.h"
+
+namespace UE
+{
+namespace GeometryFlow
+{
+
+
+class FMeshDeleteTrianglesNode : public FProcessMeshBaseNode
+{
+public:
+	static const FString InParamTriangleLists() { return TEXT("TriangleLists"); }
+
+
+public:
+	FMeshDeleteTrianglesNode()
+	{
+		AddInput(InParamTriangleLists(), MakeBasicInput<FIndexSets>());
+
+		// we can mutate input mesh
+		ConfigureInputFlags(InParamMesh(), FNodeInputFlags::Transformable());
+	}
+
+	virtual void CheckAdditionalInputs(const FNamedDataMap& DatasIn, bool& bRecomputeRequired, bool& bAllInputsValid) override
+	{
+		FindAndUpdateInputForEvaluate(InParamTriangleLists(), DatasIn, bRecomputeRequired, bAllInputsValid);
+	}
+
+
+	virtual void ProcessMesh(
+		const FNamedDataMap& DatasIn,
+		const FDynamicMesh3& MeshIn,
+		FDynamicMesh3& MeshOut) override
+	{
+		MeshOut = MeshIn;
+		DeleteTrianglesForMesh(DatasIn, MeshOut);
+	}
+
+	virtual void ProcessMeshInPlace(
+		const FNamedDataMap& DatasIn,
+		FDynamicMesh3& MeshInOut) override
+	{
+		DeleteTrianglesForMesh( DatasIn, MeshInOut);
+	}
+
+
+	void DeleteTrianglesForMesh(const FNamedDataMap& DatasIn, FDynamicMesh3& EditMesh)
+	{
+		TSafeSharedPtr<IData> TriangleSetsArg = DatasIn.FindData(InParamTriangleLists());
+		const FIndexSets& TriangleSets = TriangleSetsArg->GetDataConstRef<FIndexSets>(FIndexSets::DataTypeIdentifier);
+
+		FDynamicMeshEditor Editor(&EditMesh);
+		for (const TArray<int32>& Set : TriangleSets.IndexSets)
+		{
+			Editor.RemoveTriangles(Set, true);
+		}
+	}
+};
+
+
+
+
+
+}	// end namespace GeometryFlow
+}	//
\ No newline at end of file
diff --git a/Engine/Plugins/Experimental/GeometryFlow/Source/GeometryFlowMeshProcessing/Public/MeshProcessingNodes/MeshNormalsNodes.h b/Engine/Plugins/Experimental/GeometryFlow/Source/GeometryFlowMeshProcessing/Public/MeshProcessingNodes/MeshNormalsNodes.h
new file mode 100644
--- /dev/null
+++ b/Engine/Plugins/Experimental/GeometryFlow/Source/GeometryFlowMeshProcessing/Public/MeshProcessingNodes/MeshNormalsNodes.h
@@ -0,0 +1,120 @@
+// Copyright Epic Games, Inc. All Rights Reserved.
+
+#pragma once
+
+#include "MeshProcessingNodes/MeshProcessingBaseNodes.h"
+#include "MeshProcessingNodes/MeshProcessingDataTypes.h"
+
+#include "MeshNormals.h"
+
+
+namespace UE
+{
+namespace GeometryFlow
+{
+
+
+enum class EComputeNormalsType
+{
+	PerTriangle = 0,
+	PerVertex = 1,
+	RecomputeExistingTopology = 2,
+	FromFaceAngleThreshold = 3,
+	FromGroups = 4
+};
+
+struct FMeshNormalsSettings
+{
+	DECLARE_GEOMETRYFLOW_DATA_TYPE_IDENTIFIER(EMeshProcessingDataTypes::NormalsSettings);
+
+	EComputeNormalsType NormalsType = EComputeNormalsType::FromFaceAngleThreshold;
+
+	bool bInvert = false;
+
+	bool bAreaWeighted = true;
+	bool bAngleWeighted = true;
+
+	// for FromAngleThreshold type
+	double AngleThresholdDeg = 180.0;
+};
+GEOMETRYFLOW_DECLARE_SETTINGS_TYPES(FMeshNormalsSettings, Normals);
+
+
+
+class FComputeMeshNormalsNode : public TProcessMeshWithSettingsBaseNode<FMeshNormalsSettings>
+{
+public:
+	FComputeMeshNormalsNode() : TProcessMeshWithSettingsBaseNode<FMeshNormalsSettings>()
+	{
+		// we can mutate input mesh
+		ConfigureInputFlags(InParamMesh(), FNodeInputFlags::Transformable());
+	}
+
+
+	virtual void ProcessMesh(
+		const FNamedDataMap& DatasIn,
+		const FMeshNormalsSettings& Settings,
+		const FDynamicMesh3& MeshIn,
+		FDynamicMesh3& MeshOut) override
+	{
+		MeshOut = MeshIn;
+		ComputeNormals(Settings, MeshOut);
+	}
+
+	virtual void ProcessMeshInPlace(
+		const FNamedDataMap& DatasIn,
+		const FMeshNormalsSettings& Settings,
+		FDynamicMesh3& MeshInOut)
+	{
+		ComputeNormals(Settings, MeshInOut);
+	}
+
+
+	void ComputeNormals(const FMeshNormalsSettings& Settings, FDynamicMesh3& MeshInOut)
+	{
+		if (MeshInOut.HasAttributes() == false)
+		{
+			MeshInOut.EnableAttributes();
+		}
+		FDynamicMeshNormalOverlay* Normals = MeshInOut.Attributes()->PrimaryNormals();
+
+		if (Settings.NormalsType == EComputeNormalsType::PerTriangle)
+		{
+			ensure(Settings.bInvert == false);		// not supported
+			FMeshNormals::InitializeMeshToPerTriangleNormals(&MeshInOut);
+			return;
+		}
+		else if (Settings.NormalsType == EComputeNormalsType::PerVertex)
+		{
+			ensure(Settings.bInvert == false);		// not supported
+			FMeshNormals::InitializeOverlayToPerVertexNormals(Normals, false);
+			return;
+		}
+
+
+		if (Settings.NormalsType == EComputeNormalsType::FromFaceAngleThreshold)
+		{
+			FMeshNormals::InitializeOverlayTopologyFromOpeningAngle(&MeshInOut, Normals, Settings.AngleThresholdDeg);
+		}
+		else if (Settings.NormalsType == EComputeNormalsType::FromGroups)
+		{
+			FMeshNormals::InitializeOverlayTopologyFromFaceGroups(&MeshInOut, Normals);
+		}
+		else if (Settings.NormalsType != EComputeNormalsType::RecomputeExistingTopology)
+		{
+			ensure(false);
+		}
+
+		FMeshNormals MeshNormals(&MeshInOut);
+		MeshNormals.RecomputeOverlayNormals(Normals, Settings.bAreaWeighted, Settings.bAngleWeighted);
+		MeshNormals.CopyToOverlay(Normals, Settings.bInvert);
+	}
+
+};
+
+
+
+
+
+}	// end namespace GeometryFlow
+}	// end namespace UE
\ No newline at end of file
diff --git a/Engine/Plugins/Experimental/GeometryFlow/Source/GeometryFlowMeshProcessing/Public/MeshProcessingNodes/MeshProcessingBaseNodes.h b/Engine/Plugins/Experimental/GeometryFlow/Source/GeometryFlowMeshProcessing/Public/MeshProcessingNodes/MeshProcessingBaseNodes.h
new file mode 100644
--- /dev/null
+++ b/Engine/Plugins/Experimental/GeometryFlow/Source/GeometryFlowMeshProcessing/Public/MeshProcessingNodes/MeshProcessingBaseNodes.h
@@ -0,0 +1,254 @@
+// Copyright Epic Games, Inc. All Rights Reserved.
+
+#pragma once
+
+#include "GeometryFlowCoreNodes.h"
+#include "GeometryFlowNodeUtil.h"
+#include "DataTypes/DynamicMeshData.h"
+
+
+namespace UE
+{
+namespace GeometryFlow
+{
+
+
+
+
+
+
+
+class FProcessMeshBaseNode : public FNode
+{
+public:
+	static const FString InParamMesh() { return TEXT("Mesh"); }
+	static const FString OutParamResultMesh() { return TEXT("ResultMesh"); }
+
+public:
+	FProcessMeshBaseNode()
+	{
+		AddInput(InParamMesh(), MakeUnique<FDynamicMeshInput>());
+
+		AddOutput(OutParamResultMesh(), MakeUnique<FDynamicMeshOutput>());
+	}
+
+
+protected:
+
+	//
+	// FProcessMeshBaseNode API that subclasses must/can implement
+	//
+
+	virtual void ProcessMesh(
+		const FNamedDataMap& DatasIn,
+		const FDynamicMesh3& MeshIn,
+		FDynamicMesh3& MeshOut)
+	{
+		check(false);		// subclass must implement ProcessMesh()
+	}
+
+	virtual void ProcessMeshInPlace(
+		const FNamedDataMap& DatasIn,
+		FDynamicMesh3& MeshInOut)
+	{
+		ensureMsgf(false, TEXT("TProcessMeshBaseNode::ProcessMeshInPlace called but not defined!"));
+	}
+
+
+	virtual void CheckAdditionalInputs(const FNamedDataMap& DatasIn, bool& bRecomputeRequired, bool& bAllInputsValid)
+	{
+		// none
+	}
+
+public:
+
+	virtual void Evaluate(
+		const FNamedDataMap& DatasIn,
+		FNamedDataMap& DatasOut,
+		TUniquePtr<FEvaluationInfo>& EvaluationInfo) override
+	{
+		if (ensure(DatasOut.Contains(OutParamResultMesh())))
+		{
+			bool bAllInputsValid = true;
+			bool bRecomputeRequired = ( IsOutputAvailable(OutParamResultMesh()) == false );
+			TSafeSharedPtr<IData> MeshArg = FindAndUpdateInputForEvaluate(InParamMesh(), DatasIn, bRecomputeRequired, bAllInputsValid);
+			CheckAdditionalInputs(DatasIn, bRecomputeRequired, bAllInputsValid);
+			if (bAllInputsValid)
+			{
+				if (bRecomputeRequired)
+				{
+					bool bIsMeshMutable = DatasIn.GetDataFlags(InParamMesh()).bIsMutableData;
+					if (bIsMeshMutable)
+					{
+						UE_LOG(LogTemp, Warning, TEXT("[%s]  RECOMPUTING MeshOp In Place!"), *GetIdentifier());
+
+						FDynamicMesh3 EditableMesh;
+						MeshArg->GiveTo<FDynamicMesh3>(EditableMesh, (int)EMeshProcessingDataTypes::DynamicMesh);
+						ProcessMeshInPlace(DatasIn, EditableMesh);
+
+						// store new result
+						TSafeSharedPtr<FDataDynamicMesh> Result = MakeSafeShared<FDataDynamicMesh>();
+						Result->MoveData(MoveTemp(EditableMesh));
+						SetOutput(OutParamResultMesh(), Result);
+					}
+					else
+					{
+						UE_LOG(LogTemp, Warning, TEXT("[%s]  RECOMPUTING MeshOp"), *GetIdentifier());
+
+						// do we ever want to support using a copy of the source mesh?
+						const FDynamicMesh3& SourceMesh = MeshArg->GetDataConstRef<FDynamicMesh3>((int)EMeshProcessingDataTypes::DynamicMesh);
+
+						// run mesh processing
+						FDynamicMesh3 ResultMesh;
+						ProcessMesh(DatasIn, SourceMesh, ResultMesh);
+
+						// store new result
+						TSafeSharedPtr<FDataDynamicMesh> Result = MakeSafeShared<FDataDynamicMesh>();
+						Result->MoveData(MoveTemp(ResultMesh));
+						SetOutput(OutParamResultMesh(), Result);
+					}
+
+					EvaluationInfo->CountCompute(this);
+				}
+				DatasOut.SetData(OutParamResultMesh(), GetOutput(OutParamResultMesh()));
+			}
+		}
+	}
+};
+
+
+
+
+
+
+
+
+
+
+
+template<typename SettingsType>
+class TProcessMeshWithSettingsBaseNode : public FNode
+{
+public:
+	using SettingsDataType = TMovableData<SettingsType, SettingsType::DataTypeIdentifier>;
+
+public:
+	static const FString InParamMesh() { return TEXT("Mesh"); }
+	static const FString InParamSettings() { return TEXT("Settings"); }
+	static const FString OutParamResultMesh() { return TEXT("ResultMesh"); }
+
+protected:
+	TProcessMeshWithSettingsBaseNode()
+	{
+		AddInput(InParamMesh(), MakeUnique<FDynamicMeshInput>());
+		AddInput(InParamSettings(), MakeUnique<TBasicNodeInput<SettingsType, SettingsType::DataTypeIdentifier>>());
+
+		AddOutput(OutParamResultMesh(), MakeUnique<FDynamicMeshOutput>());
+	}
+
+
+	virtual void ProcessMesh(
+		const FNamedDataMap& DatasIn,
+		const SettingsType& SettingsIn,
+		const FDynamicMesh3& MeshIn,
+		FDynamicMesh3& MeshOut)
+	{
+	}
+
+	virtual void ProcessMeshInPlace(
+		const FNamedDataMap& DatasIn,
+		const SettingsType& SettingsIn,
+		FDynamicMesh3& MeshInOut)
+	{
+		ensureMsgf(false, TEXT("TProcessMeshWithSettingsBaseNode::ProcessMeshInPlace called but not defined!"));
+	}
+
+
+	virtual void CheckAdditionalInputs(const FNamedDataMap& DatasIn, bool& bRecomputeRequired, bool& bAllInputsValid)
+	{
+		// none
+	}
+
+public:
+
+	virtual void Evaluate(
+		const FNamedDataMap& DatasIn,
+		FNamedDataMap& DatasOut,
+		TUniquePtr<FEvaluationInfo>& EvaluationInfo) override
+	{
+		if (ensure(DatasOut.Contains(OutParamResultMesh())))
+		{
+			bool bAllInputsValid = true;
+			bool bRecomputeRequired = ( IsOutputAvailable(OutParamResultMesh()) == false );
+			TSafeSharedPtr<IData> MeshArg = FindAndUpdateInputForEvaluate(InParamMesh(), DatasIn, bRecomputeRequired, bAllInputsValid);
+			TSafeSharedPtr<IData> SettingsArg = FindAndUpdateInputForEvaluate(InParamSettings(), DatasIn, bRecomputeRequired, bAllInputsValid);
+			CheckAdditionalInputs(DatasIn, bRecomputeRequired, bAllInputsValid);
+			if (bAllInputsValid)
+			{
+				if (bRecomputeRequired)
+				{
+					// always make a copy of settings
+					SettingsType Settings;
+					SettingsArg->GetDataCopy(Settings, SettingsType::DataTypeIdentifier);
+
+					bool bIsMeshMutable = DatasIn.GetDataFlags(InParamMesh()).bIsMutableData;
+					if (bIsMeshMutable)
+					{
+						UE_LOG(LogTemp, Warning, TEXT("[%s]  RECOMPUTING MeshOp In Place!"), *GetIdentifier());
+
+						FDynamicMesh3 EditableMesh;
+						MeshArg->GiveTo<FDynamicMesh3>(EditableMesh, (int)EMeshProcessingDataTypes::DynamicMesh);
+						ProcessMeshInPlace(DatasIn, Settings, EditableMesh);
+
+						// store new result
+						TSafeSharedPtr<FDataDynamicMesh> Result = MakeSafeShared<FDataDynamicMesh>();
+						Result->MoveData(MoveTemp(EditableMesh));
+						SetOutput(OutParamResultMesh(), Result);
+					}
+					else
+					{
+						UE_LOG(LogTemp, Warning, TEXT("[%s]  RECOMPUTING MeshOp"), *GetIdentifier());
+
+						// do we ever want to support using a copy of the source mesh?
+						const FDynamicMesh3& SourceMesh = MeshArg->GetDataConstRef<FDynamicMesh3>((int)EMeshProcessingDataTypes::DynamicMesh);
+
+						// run mesh processing
+						FDynamicMesh3 ResultMesh;
+						ProcessMesh(DatasIn, Settings, SourceMesh, ResultMesh);
+
+						// store new result
+						TSafeSharedPtr<FDataDynamicMesh> Result = MakeSafeShared<FDataDynamicMesh>();
+						Result->MoveData(MoveTemp(ResultMesh));
+						SetOutput(OutParamResultMesh(), Result);
+					}
+
+					EvaluationInfo->CountCompute(this);
+				}
+				DatasOut.SetData(OutParamResultMesh(), GetOutput(OutParamResultMesh()));
+			}
+		}
+	}
+};
+
+
+
+
+
+
+
+
+
+
+
+
+
+
+
+
+
+
+
+}	// end namespace GeometryFlow
+}	// end namespace UE
+
+
diff --git a/Engine/Plugins/Experimental/GeometryFlow/Source/GeometryFlowMeshProcessing/Public/MeshProcessingNodes/MeshProcessingDataTypes.h b/Engine/Plugins/Experimental/GeometryFlow/Source/GeometryFlowMeshProcessing/Public/MeshProcessingNodes/MeshProcessingDataTypes.h
new file mode 100644
--- /dev/null
+++ b/Engine/Plugins/Experimental/GeometryFlow/Source/GeometryFlowMeshProcessing/Public/MeshProcessingNodes/MeshProcessingDataTypes.h
@@ -0,0 +1,52 @@
+// Copyright Epic Games, Inc. All Rights Reserved.
+
+#pragma once
+
+#include "CoreMinimal.h"
+#include "GeometryFlowTypes.h"
+
+namespace UE
+{
+namespace GeometryFlow
+{
+
+enum class EMeshProcessingDataTypes
+{
+	DynamicMesh = (int)EDataTypes::BaseMeshProcessingTypes + 1,
+	MeshNormalSet = (int)EDataTypes::BaseMeshProcessingTypes + 2,
+	MeshTangentSet = (int)EDataTypes::BaseMeshProcessingTypes + 3,
+
+	BakingCache = (int)EDataTypes::BaseMeshProcessingTypes + 10,
+	TextureImage = (int)EDataTypes::BaseMeshProcessingTypes + 11,
+	NormalMapImage = (int)EDataTypes::BaseMeshProcessingTypes + 12,
+	CollisionGeometry = (int)EDataTypes::BaseMeshProcessingTypes + 13,
+	IndexSets = (int)EDataTypes::BaseMeshProcessingTypes + 14,
+
+
+	SolidifySettings = (int)EDataTypes::BaseMeshProcessingTypes + 100,
+	VoxMorphologyOpSettings = (int)EDataTypes::BaseMeshProcessingTypes + 101,
+
+	SimplifySettings = (int)EDataTypes::BaseMeshProcessingTypes + 110,
+
+
+
+	NormalsSettings = (int)EDataTypes::BaseMeshProcessingTypes + 120,
+	TangentsSettings = (int)EDataTypes::BaseMeshProcessingTypes + 121,
+
+
+	RecalculateUVsSettings = (int)EDataTypes::BaseMeshProcessingTypes + 150,
+	RepackUVsSettings = (int)EDataTypes::BaseMeshProcessingTypes + 151,
+
+
+	MakeBakingCacheSettings = (int)EDataTypes::BaseMeshProcessingTypes + 160,
+	BakeNormalMapSettings = (int)EDataTypes::BaseMeshProcessingTypes + 161,
+	BakeTextureImageSettings = (int)EDataTypes::BaseMeshProcessingTypes + 162,
+
+
+	GenerateCollisionConvexHullsSettings = (int)EDataTypes::BaseMeshProcessingTypes + 200,
+};
+
+
+
+}	// end namespace GeometryFlow
+}	// en
\ No newline at end of file
diff --git a/Engine/Plugins/Experimental/GeometryFlow/Source/GeometryFlowMeshProcessing/Public/MeshProcessingNodes/MeshRecalculateUVsNode.h b/Engine/Plugins/Experimental/GeometryFlow/Source/GeometryFlowMeshProcessing/Public/MeshProcessingNodes/MeshRecalculateUVsNode.h
new file mode 100644
--- /dev/null
+++ b/Engine/Plugins/Experimental/GeometryFlow/Source/GeometryFlowMeshProcessing/Public/MeshProcessingNodes/MeshRecalculateUVsNode.h
@@ -0,0 +1,68 @@
+// Copyright Epic Games, Inc. All Rights Reserved.
+
+#pragma once
+
+#include "MeshProcessingNodes/MeshProcessingBaseNodes.h"
+#include "MeshProcessingNodes/MeshProcessingDataTypes.h"
+
+namespace UE
+{
+namespace GeometryFlow
+{
+
+enum class ERecalculateUVsUnwrapType : uint8
+{
+	Auto = 0,
+	ExpMap = 1,
+	Conformal = 2
+};
+
+
+struct GEOMETRYFLOWMESHPROCESSING_API FMeshRecalculateUVsSettings
+{
+	DECLARE_GEOMETRYFLOW_DATA_TYPE_IDENTIFIER(EMeshProcessingDataTypes::RecalculateUVsSettings);
+
+	ERecalculateUVsUnwrapType UnwrapType = ERecalculateUVsUnwrapType::Auto;
+	int32 UVLayer = 0;
+
+};
+GEOMETRYFLOW_DECLARE_SETTINGS_TYPES(FMeshRecalculateUVsSettings, MeshRecalculateUVs);
+
+
+
+class GEOMETRYFLOWMESHPROCESSING_API FMeshRecalculateUVsNode : public TProcessMeshWithSettingsBaseNode<FMeshRecalculateUVsSettings>
+{
+public:
+	FMeshRecalculateUVsNode() : TProcessMeshWithSettingsBaseNode<FMeshRecalculateUVsSettings>()
+	{
+		// we can mutate input mesh
+		ConfigureInputFlags(InParamMesh(), FNodeInputFlags::Transformable());
+	}
+
+	virtual void ProcessMesh(
+		const FNamedDataMap& DatasIn,
+		const FMeshRecalculateUVsSettings& Settings,
+		const FDynamicMesh3& MeshIn,
+		FDynamicMesh3& MeshOut) override
+	{
+		MeshOut = MeshIn;
+		RecalculateUVsOnMesh(MeshOut, Settings);
+	}
+
+	virtual void ProcessMeshInPlace(
+		const FNamedDataMap& DatasIn,
+		const FMeshRecalculateUVsSettings& Settings,
+		FDynamicMesh3& MeshInOut)
+	{
+		RecalculateUVsOnMesh(MeshInOut, Settings);
+	}
+
+	void RecalculateUVsOnMesh(FDynamicMesh3& EditMesh, const FMeshRecalculateUVsSettings& Settings);
+};
+
+
+
+
+
+}	// end namespace GeometryFlow
+}	// end 
\ No newline at end of file
diff --git a/Engine/Plugins/Experimental/GeometryFlow/Source/GeometryFlowMeshProcessing/Public/MeshProcessingNodes/MeshRepackUVsNode.h b/Engine/Plugins/Experimental/GeometryFlow/Source/GeometryFlowMeshProcessing/Public/MeshProcessingNodes/MeshRepackUVsNode.h
new file mode 100644
--- /dev/null
+++ b/Engine/Plugins/Experimental/GeometryFlow/Source/GeometryFlowMeshProcessing/Public/MeshProcessingNodes/MeshRepackUVsNode.h
@@ -0,0 +1,67 @@
+// Copyright Epic Games, Inc. All Rights Reserved.
+
+#pragma once
+
+#include "MeshProcessingNodes/MeshProcessingBaseNodes.h"
+#include "MeshProcessingNodes/MeshProcessingDataTypes.h"
+
+namespace UE
+{
+namespace GeometryFlow
+{
+
+
+
+struct GEOMETRYFLOWMESHPROCESSING_API FMeshRepackUVsSettings
+{
+	DECLARE_GEOMETRYFLOW_DATA_TYPE_IDENTIFIER(EMeshProcessingDataTypes::RepackUVsSettings);
+
+	int32 UVLayer = 0;
+
+	int32 TextureResolution = 512;
+	int32 GutterSize = 1;
+	bool bAllowFlips = false;
+
+	FVector2f UVScale = FVector2f::One();
+	FVector2f UVTranslation = FVector2f::Zero();
+};
+GEOMETRYFLOW_DECLARE_SETTINGS_TYPES(FMeshRepackUVsSettings, MeshRepackUVs);
+
+
+
+class GEOMETRYFLOWMESHPROCESSING_API FMeshRepackUVsNode : public TProcessMeshWithSettingsBaseNode<FMeshRepackUVsSettings>
+{
+public:
+	FMeshRepackUVsNode() : TProcessMeshWithSettingsBaseNode<FMeshRepackUVsSettings>()
+	{
+		// we can mutate input mesh
+		ConfigureInputFlags(InParamMesh(), FNodeInputFlags::Transformable());
+	}
+
+	virtual void ProcessMesh(
+		const FNamedDataMap& DatasIn,
+		const FMeshRepackUVsSettings& Settings,
+		const FDynamicMesh3& MeshIn,
+		FDynamicMesh3& MeshOut) override
+	{
+		MeshOut = MeshIn;
+		RepackUVsForMesh(MeshOut, Settings);
+	}
+
+	virtual void ProcessMeshInPlace(
+		const FNamedDataMap& DatasIn,
+		const FMeshRepackUVsSettings& Settings,
+		FDynamicMesh3& MeshInOut)
+	{
+		RepackUVsForMesh(MeshInOut, Settings);
+	}
+
+	void RepackUVsForMesh(FDynamicMesh3& EditMesh, const FMeshRepackUVsSettings& Settings);
+};
+
+
+
+
+
+}	// end namespace GeometryFlow
+}	// end 
\ No newline at end of file
diff --git a/Engine/Plugins/Experimental/GeometryFlow/Source/GeometryFlowMeshProcessing/Public/MeshProcessingNodes/MeshSimplifyNode.h b/Engine/Plugins/Experimental/GeometryFlow/Source/GeometryFlowMeshProcessing/Public/MeshProcessingNodes/MeshSimplifyNode.h
new file mode 100644
--- /dev/null
+++ b/Engine/Plugins/Experimental/GeometryFlow/Source/GeometryFlowMeshProcessing/Public/MeshProcessingNodes/MeshSimplifyNode.h
@@ -0,0 +1,168 @@
+// Copyright Epic Games, Inc. All Rights Reserved.
+
+#pragma once
+
+#include "MeshProcessingNodes/MeshProcessingBaseNodes.h"
+#include "MeshProcessingNodes/MeshProcessingDataTypes.h"
+
+#include "MeshSimplification.h"
+#include "MeshConstraintsUtil.h"
+
+
+namespace UE
+{
+namespace GeometryFlow
+{
+
+
+enum class EMeshSimplifyType
+{
+	Standard = 0,
+	VolumePreserving = 1,
+	AttributeAware = 2
+};
+
+enum class EMeshSimplifyTargetType
+{
+	TriangleCount = 0,
+	VertexCount = 1,
+	TrianglePercentage = 2
+};
+
+struct FMeshSimplifySettings
+{
+	DECLARE_GEOMETRYFLOW_DATA_TYPE_IDENTIFIER(EMeshProcessingDataTypes::SimplifySettings);
+
+	EMeshSimplifyType SimplifyType = EMeshSimplifyType::AttributeAware;
+	EMeshSimplifyTargetType TargetType = EMeshSimplifyTargetType::TrianglePercentage;
+
+	int TargetCount = 100;
+	float TargetFraction = 0.5;
+
+	bool bDiscardAttributes = false;
+	bool bPreventNormalFlips = true;
+	bool bPreserveSharpEdges = false;
+	bool bAllowSeamCollapse = true;
+	bool bAllowSeamSplits = true;
+
+	EEdgeRefineFlags MeshBoundaryConstraints = EEdgeRefineFlags::NoConstraint;
+	EEdgeRefineFlags GroupBorderConstraints = EEdgeRefineFlags::NoConstraint;
+	EEdgeRefineFlags MaterialBorderConstraints = EEdgeRefineFlags::NoConstraint;
+};
+GEOMETRYFLOW_DECLARE_SETTINGS_TYPES(FMeshSimplifySettings, Simplify);
+
+
+class FSimplifyMeshNode : public TProcessMeshWithSettingsBaseNode<FMeshSimplifySettings>
+{
+public:
+	FSimplifyMeshNode() : TProcessMeshWithSettingsBaseNode<FMeshSimplifySettings>()
+	{
+		// we can mutate input mesh
+		ConfigureInputFlags(InParamMesh(), FNodeInputFlags::Transformable());
+	}
+
+	virtual void ProcessMesh(
+		const FNamedDataMap& DatasIn,
+		const FMeshSimplifySettings& Settings,
+		const FDynamicMesh3& MeshIn,
+		FDynamicMesh3& MeshOut) override
+	{
+		MeshOut.Copy(MeshIn, true, true, true, !Settings.bDiscardAttributes);
+		ApplySimplify(Settings, MeshOut);
+	}
+
+	virtual void ProcessMeshInPlace(
+		const FNamedDataMap& DatasIn,
+		const FMeshSimplifySettings& Settings,
+		FDynamicMesh3& MeshInOut)
+	{
+		ApplySimplify(Settings, MeshInOut);
+	}
+
+
+	template<typename SimplifierType> 
+	void DoSimplifyOfType(const FMeshSimplifySettings& Settings, FDynamicMesh3* TargetMesh)
+	{
+		SimplifierType Simplifier(TargetMesh);
+
+		Simplifier.ProjectionMode = SimplifierType::ETargetProjectionMode::NoProjection;
+		Simplifier.DEBUG_CHECK_LEVEL = 0;
+
+		Simplifier.bAllowSeamCollapse = Settings.bAllowSeamCollapse;
+		if (Simplifier.bAllowSeamCollapse)
+		{
+			Simplifier.SetEdgeFlipTolerance(1.e-5);
+
+			// eliminate any bowties that might have formed on UV seams.
+			if (FDynamicMeshAttributeSet* Attributes = TargetMesh->Attributes())
+			{
+				// @todo parallelize over NumUVLayers?
+				int NumUVLayers = Attributes->NumUVLayers();
+				for (int i = 0; i < NumUVLayers; ++i)
+				{
+					Attributes->GetUVLayer(i)->SplitBowties();
+				}
+				Attributes->PrimaryNormals()->SplitBowties();
+
+			}
+		}
+
+		FMeshConstraints Constraints;
+		FMeshConstraintsUtil::ConstrainAllBoundariesAndSeams(Constraints, *TargetMesh,
+			Settings.MeshBoundaryConstraints,
+			Settings.GroupBorderConstraints,
+			Settings.MaterialBorderConstraints,
+			Settings.bAllowSeamSplits, !Settings.bPreserveSharpEdges, Settings.bAllowSeamCollapse);
+		Simplifier.SetExternalConstraints(MoveTemp(Constraints));
+
+		if (Settings.TargetType == EMeshSimplifyTargetType::TrianglePercentage)
+		{
+			int32 UseTarget = FMath::Max(4, (int)(Settings.TargetFraction * (double)TargetMesh->TriangleCount()));
+			Simplifier.SimplifyToTriangleCount(UseTarget);
+		}
+		else if (Settings.TargetType == EMeshSimplifyTargetType::TriangleCount)
+		{
+			Simplifier.SimplifyToTriangleCount( FMath::Max(1,Settings.TargetCount) );
+		}
+		else if (Settings.TargetType == EMeshSimplifyTargetType::VertexCount)
+		{
+			Simplifier.SimplifyToVertexCount( FMath::Max(3, Settings.TargetCount) );
+		}
+		else
+		{
+			check(false);
+		}
+	}
+
+
+	void ApplySimplify(const FMeshSimplifySettings& Settings, FDynamicMesh3& MeshInOut)
+	{
+		if (Settings.bDiscardAttributes)
+		{
+			MeshInOut.DiscardAttributes();
+		}
+
+		if (Settings.SimplifyType == EMeshSimplifyType::Standard)
+		{
+			DoSimplifyOfType<FQEMSimplification>(Settings, &MeshInOut);
+		}
+		else if (Settings.SimplifyType == EMeshSimplifyType::VolumePreserving)
+		{
+			DoSimplifyOfType<FVolPresMeshSimplification>(Settings, &MeshInOut);
+		}
+		else if (Settings.SimplifyType == EMeshSimplifyType::AttributeAware)
+		{
+			DoSimplifyOfType<FAttrMeshSimplification>(Settings, &MeshInOut);
+		}
+		else
+		{
+			check(false);
+		}
+	}
+
+};
+
+
+
+}	// end namespace GeometryFlow
+}	// end namespace UE
\ No newline at end of file
diff --git a/Engine/Plugins/Experimental/GeometryFlow/Source/GeometryFlowMeshProcessing/Public/MeshProcessingNodes/MeshSolidifyNode.h b/Engine/Plugins/Experimental/GeometryFlow/Source/GeometryFlowMeshProcessing/Public/MeshProcessingNodes/MeshSolidifyNode.h
new file mode 100644
--- /dev/null
+++ b/Engine/Plugins/Experimental/GeometryFlow/Source/GeometryFlowMeshProcessing/Public/MeshProcessingNodes/MeshSolidifyNode.h
@@ -0,0 +1,64 @@
+// Copyright Epic Games, Inc. All Rights Reserved.
+
+#pragma once
+
+#include "MeshProcessingNodes/MeshProcessingBaseNodes.h"
+#include "MeshProcessingNodes/MeshProcessingDataTypes.h"
+#include "DynamicMeshAABBTree3.h"
+#include "Implicit/Solidify.h"
+
+
+namespace UE
+{
+namespace GeometryFlow
+{
+
+
+
+struct FMeshSolidifySettings
+{
+	DECLARE_GEOMETRYFLOW_DATA_TYPE_IDENTIFIER(EMeshProcessingDataTypes::SolidifySettings);
+
+	int VoxelResolution = 64;
+	float WindingThreshold = 0.5;
+	int SurfaceConvergeSteps = 5;
+	float ExtendBounds = 2.0f;
+};
+GEOMETRYFLOW_DECLARE_SETTINGS_TYPES(FMeshSolidifySettings, Solidify);
+
+
+class FSolidifyMeshNode : public TProcessMeshWithSettingsBaseNode<FMeshSolidifySettings>
+{
+public:
+
+	virtual void ProcessMesh(
+		const FNamedDataMap& DatasIn,
+		const FMeshSolidifySettings& SettingsIn,
+		const FDynamicMesh3& MeshIn,
+		FDynamicMesh3& MeshOut) override
+	{
+		ApplySolidify(MeshIn, SettingsIn, MeshOut);
+	}
+
+	void ApplySolidify(const FDynamicMesh3& MeshIn, const FMeshSolidifySettings& Settings, FDynamicMesh3& MeshOut)
+	{
+		FAxisAlignedBox3d Bounds = MeshIn.GetBounds();
+
+		FDynamicMeshAABBTree3 MeshBVTree(&MeshIn);
+		TFastWindingTree<FDynamicMesh3> FastWinding(&MeshBVTree);
+
+		TImplicitSolidify<FDynamicMesh3> SolidifyCalc(&MeshIn, &MeshBVTree, &FastWinding);
+
+		SolidifyCalc.SetCellSizeAndExtendBounds(Bounds, Settings.ExtendBounds, Settings.VoxelResolution);
+		SolidifyCalc.WindingThreshold = Settings.WindingThreshold;
+		SolidifyCalc.SurfaceSearchSteps = Settings.SurfaceConvergeSteps;
+		SolidifyCalc.bSolidAtBoundaries = true;
+		SolidifyCalc.ExtendBounds = Settings.ExtendBounds;
+
+		MeshOut.Copy(&SolidifyCalc.Generate());
+	}
+};
+
+
+}	// end namespace GeometryFlow
+}	// end namespace UE
\ No newline at end of file
diff --git a/Engine/Plugins/Experimental/GeometryFlow/Source/GeometryFlowMeshProcessing/Public/MeshProcessingNodes/MeshTangentsNodes.h b/Engine/Plugins/Experimental/GeometryFlow/Source/GeometryFlowMeshProcessing/Public/MeshProcessingNodes/MeshTangentsNodes.h
new file mode 100644
--- /dev/null
+++ b/Engine/Plugins/Experimental/GeometryFlow/Source/GeometryFlowMeshProcessing/Public/MeshProcessingNodes/MeshTangentsNodes.h
@@ -0,0 +1,85 @@
+// Copyright Epic Games, Inc. All Rights Reserved.
+
+#pragma once
+
+#include "BaseNodes/TransformerWithSettingsNode.h"
+#include "MeshProcessingNodes/MeshProcessingBaseNodes.h"
+#include "MeshProcessingNodes/MeshProcessingDataTypes.h"
+
+#include "MeshTangents.h"
+
+
+namespace UE
+{
+namespace GeometryFlow
+{
+
+
+enum class EComputeTangentsType
+{
+	PerTriangle = 0,
+	FastMikkT = 1
+};
+
+
+struct FMeshTangentsSettings
+{
+	DECLARE_GEOMETRYFLOW_DATA_TYPE_IDENTIFIER(EMeshProcessingDataTypes::TangentsSettings);
+
+	EComputeTangentsType TangentsType = EComputeTangentsType::FastMikkT;
+
+	int UVLayer = 0;
+};
+GEOMETRYFLOW_DECLARE_SETTINGS_TYPES(FMeshTangentsSettings, Tangents);
+
+
+
+
+class FComputeMeshTangentsNode : public TTransformerWithSettingsNode<
+	FDynamicMesh3, (int)EMeshProcessingDataTypes::DynamicMesh,
+	FMeshTangentsSettings, (int)EMeshProcessingDataTypes::TangentsSettings,
+	FMeshTangentsd, (int)EMeshProcessingDataTypes::MeshTangentSet>
+{
+public:
+	FComputeMeshTangentsNode() : TTransformerWithSettingsNode()
+	{
+	}
+
+	virtual void ComputeOutput(
+		const FNamedDataMap& DatasIn,
+		const FMeshTangentsSettings& Settings,
+		const FDynamicMesh3& MeshIn,
+		FMeshTangentsd& TangentsOut)
+	{
+		if (ensure(MeshIn.HasAttributes()) == false) return;
+
+		// handle UV layer issues, missing attributes, etc
+		int UseUVLayer = Settings.UVLayer;
+		bool bValidUVLayer = (UseUVLayer >= 0 && UseUVLayer < MeshIn.Attributes()->NumUVLayers());
+		if (ensure(bValidUVLayer) == false)
+		{
+			UseUVLayer = FMath::Clamp(UseUVLayer, 0, MeshIn.Attributes()->NumUVLayers() - 1);
+		}
+
+		FComputeTangentsOptions Options;
+		Options.bAveraged = (Settings.TangentsType == EComputeTangentsType::FastMikkT);
+
+		TangentsOut.SetMesh(&MeshIn);
+		TangentsOut.ComputeTriVertexTangents(
+			MeshIn.Attributes()->PrimaryNormals(),
+			MeshIn.Attributes()->GetUVLayer(UseUVLayer),
+			Options);
+
+		// clear output mesh reference
+		TangentsOut.SetMesh(nullptr);
+	}
+
+
+};
+
+
+
+
+
+}	// end namespace GeometryFlow
+}	// end namespace UE
\ No newline at end of file
diff --git a/Engine/Plugins/Experimental/GeometryFlow/Source/GeometryFlowMeshProcessing/Public/MeshProcessingNodes/MeshVoxMorphologyNode.h b/Engine/Plugins/Experimental/GeometryFlow/Source/GeometryFlowMeshProcessing/Public/MeshProcessingNodes/MeshVoxMorphologyNode.h
new file mode 100644
--- /dev/null
+++ b/Engine/Plugins/Experimental/GeometryFlow/Source/GeometryFlowMeshProcessing/Public/MeshProcessingNodes/MeshVoxMorphologyNode.h
@@ -0,0 +1,78 @@
+// Copyright Epic Games, Inc. All Rights Reserved.
+
+#pragma once
+
+#include "MeshProcessingNodes/MeshProcessingBaseNodes.h"
+#include "MeshProcessingNodes/MeshProcessingDataTypes.h"
+
+#include "DynamicMeshAABBTree3.h"
+#include "Implicit/Morphology.h"
+
+
+namespace UE
+{
+namespace GeometryFlow
+{
+
+
+
+struct FVoxMorphologyOpSettings
+{
+	DECLARE_GEOMETRYFLOW_DATA_TYPE_IDENTIFIER(EMeshProcessingDataTypes::VoxMorphologyOpSettings);
+
+	int VoxelResolution = 64;
+	double Distance = 1.0;
+};
+typedef FVoxMorphologyOpSettings FVoxOffsetSettings;
+typedef FVoxMorphologyOpSettings FVoxClosureSettings;
+typedef FVoxMorphologyOpSettings FVoxOpeningSettings;
+GEOMETRYFLOW_DECLARE_SETTINGS_TYPES(FVoxMorphologyOpSettings, VoxOffset);
+GEOMETRYFLOW_DECLARE_SETTINGS_TYPES(FVoxMorphologyOpSettings, VoxClosure);
+GEOMETRYFLOW_DECLARE_SETTINGS_TYPES(FVoxMorphologyOpSettings, VoxOpening);
+
+
+template<TImplicitMorphology<FDynamicMesh3>::EMorphologyOp MorphologyOp>
+class TVoxMorphologyMeshNode : public TProcessMeshWithSettingsBaseNode<FVoxMorphologyOpSettings>
+{
+public:
+
+
+	virtual void ProcessMesh(
+		const FNamedDataMap& DatasIn,
+		const FVoxMorphologyOpSettings& SettingsIn,
+		const FDynamicMesh3& MeshIn,
+		FDynamicMesh3& MeshOut) override
+	{
+		ApplyMorphology(MeshIn, SettingsIn, MeshOut);
+	}
+
+	void ApplyMorphology(const FDynamicMesh3& MeshIn, const FVoxMorphologyOpSettings& Settings, FDynamicMesh3& MeshOut)
+	{
+		FAxisAlignedBox3d Bounds = MeshIn.GetBounds();
+		FDynamicMeshAABBTree3 MeshBVTree(&MeshIn);
+
+		TImplicitMorphology<FDynamicMesh3> ImplicitMorphology;
+		if (MorphologyOp == TImplicitMorphology<FDynamicMesh3>::EMorphologyOp::Dilate && Settings.Distance < 0)
+		{
+			ImplicitMorphology.MorphologyOp = TImplicitMorphology<FDynamicMesh3>::EMorphologyOp::Contract;
+		}
+		else
+		{
+			ImplicitMorphology.MorphologyOp = MorphologyOp;
+		}
+		ImplicitMorphology.Source = &MeshIn;
+		ImplicitMorphology.SourceSpatial = &MeshBVTree;
+		ImplicitMorphology.SetCellSizesAndDistance(Bounds, Settings.Distance, Settings.VoxelResolution, Settings.VoxelResolution);
+
+		MeshOut.Copy(&ImplicitMorphology.Generate());
+	}
+
+};
+
+typedef TVoxMorphologyMeshNode<TImplicitMorphology<FDynamicMesh3>::EMorphologyOp::Dilate> FVoxDilateMeshNode;
+typedef TVoxMorphologyMeshNode<TImplicitMorphology<FDynamicMesh3>::EMorphologyOp::Close> FVoxClosureMeshNode;
+typedef TVoxMorphologyMeshNode<TImplicitMorphology<FDynamicMesh3>::EMorphologyOp::Open> FVoxOpeningMeshNode;
+
+
+}	// end namespace GeometryFlow
+}	// end namespace UE
\ No newline at end of file
diff --git a/Engine/Plugins/Experimental/GeometryFlow/Source/GeometryFlowMeshProcessing/Public/PhysicsNodes/GenerateConvexHullsCollisionNode.h b/Engine/Plugins/Experimental/GeometryFlow/Source/GeometryFlowMeshProcessing/Public/PhysicsNodes/GenerateConvexHullsCollisionNode.h
new file mode 100644
--- /dev/null
+++ b/Engine/Plugins/Experimental/GeometryFlow/Source/GeometryFlowMeshProcessing/Public/PhysicsNodes/GenerateConvexHullsCollisionNode.h
@@ -0,0 +1,65 @@
+// Copyright Epic Games, Inc. All Rights Reserved.
+
+#pragma once
+
+#include "GeometryFlowCoreNodes.h"
+#include "GeometryFlowNodeUtil.h"
+#include "MeshProcessingNodes/MeshProcessingDataTypes.h"
+#include "DataTypes/CollisionGeometryData.h"
+#include "DataTypes/IndexSetsData.h"
+
+
+namespace UE
+{
+namespace GeometryFlow
+{
+
+
+
+
+struct GEOMETRYFLOWMESHPROCESSING_API FGenerateConvexHullsCollisionSettings
+{
+	DECLARE_GEOMETRYFLOW_DATA_TYPE_IDENTIFIER(EMeshProcessingDataTypes::GenerateCollisionConvexHullsSettings);
+
+	int32 SimplifyToTriangleCount = 50;
+	
+};
+GEOMETRYFLOW_DECLARE_SETTINGS_TYPES(FGenerateConvexHullsCollisionSettings, GenerateConvexHullsCollision);
+
+
+
+class GEOMETRYFLOWMESHPROCESSING_API FGenerateConvexHullsCollisionNode : public FNode
+{
+protected:
+	using SettingsDataType = TMovableData<FGenerateConvexHullsCollisionSettings, FGenerateConvexHullsCollisionSettings::DataTypeIdentifier>;
+
+public:
+	static const FString InParamMesh() { return TEXT("Mesh"); }
+	static const FString InParamIndexSets() { return TEXT("TriangleSets"); }
+	static const FString InParamSettings() { return TEXT("Settings"); }
+
+	static const FString OutParamGeometry() { return TEXT("Geometry"); }
+
+public:
+	FGenerateConvexHullsCollisionNode()
+	{
+		AddInput(InParamMesh(), MakeUnique<FDynamicMeshInput>());
+		AddInput(InParamIndexSets(), MakeBasicInput<FIndexSets>());
+		AddInput(InParamSettings(), MakeBasicInput<FGenerateConvexHullsCollisionSettings>());
+
+		AddOutput(OutParamGeometry(), MakeBasicOutput<FCollisionGeometry>());
+	}
+
+	virtual void Evaluate(
+		const FNamedDataMap& DatasIn,
+		FNamedDataMap& DatasOut,
+		TUniquePtr<FEvaluationInfo>& EvaluationInfo) override;
+};
+
+
+
+
+
+
+}	// end namespace GeometryFlow
+}	// end namespace UE
\ No newline at end of file
diff --git a/Engine/Plugins/Experimental/GeometryFlow/Source/GeometryFlowMeshProcessingEditor/GeometryFlowMeshProcessingEditor.Build.cs b/Engine/Plugins/Experimental/GeometryFlow/Source/GeometryFlowMeshProcessingEditor/GeometryFlowMeshProcessingEditor.Build.cs
new file mode 100644
--- /dev/null
+++ b/Engine/Plugins/Experimental/GeometryFlow/Source/GeometryFlowMeshProcessingEditor/GeometryFlowMeshProcessingEditor.Build.cs
@@ -0,0 +1,53 @@
+// Copyright Epic Games, Inc. All Rights Reserved.
+
+using UnrealBuildTool;
+
+public class GeometryFlowMeshProcessingEditor : ModuleRules
+{
+	public GeometryFlowMeshProcessingEditor(ReadOnlyTargetRules Target) : base(Target)
+	{
+		PCHUsage = ModuleRules.PCHUsageMode.UseExplicitOrSharedPCHs;
+		
+		PublicIncludePaths.AddRange(
+			new string[] {
+				// ... add public include paths required here ...
+			}
+			);
+				
+		
+		PrivateIncludePaths.AddRange(
+			new string[] {
+				// ... add other private include paths required here ...
+			}
+			);
+			
+		
+		PublicDependencyModuleNames.AddRange(
+			new string[]
+			{
+				"Core",
+				"GeometricObjects",
+				"GeometryAlgorithms",
+				"DynamicMesh",
+				"GeometryFlowCore",
+				"GeometryFlowMeshProcessing"
+			}
+            );
+			
+		
+		PrivateDependencyModuleNames.AddRange(
+			new string[]
+			{
+				"ModelingOperatorsEditorOnly"
+			}
+			);
+		
+		
+		DynamicallyLoadedModuleNames.AddRange(
+			new string[]
+			{
+				// ... add any modules that your module loads dynamically here ...
+			}
+			);
+	}
+}
diff --git a/Engine/Plugins/Experimental/GeometryFlow/Source/GeometryFlowMeshProcessingEditor/Private/GeometryFlowMeshProcessingEditorModule.cpp b/Engine/Plugins/Experimental/GeometryFlow/Source/GeometryFlowMeshProcessingEditor/Private/GeometryFlowMeshProcessingEditorModule.cpp
new file mode 100644
--- /dev/null
+++ b/Engine/Plugins/Experimental/GeometryFlow/Source/GeometryFlowMeshProcessingEditor/Private/GeometryFlowMeshProcessingEditorModule.cpp
@@ -0,0 +1,20 @@
+// Copyright Epic Games, Inc. All Rights Reserved.
+
+#include "GeometryFlowMeshProcessingEditorModule.h"
+
+#define LOCTEXT_NAMESPACE "FGeometryFlowMeshProcessingEditorModule"
+
+void FGeometryFlowMeshProcessingEditorModule::StartupModule()
+{
+	// This code will execute after your module is loaded into memory; the exact timing is specified in the .uplugin file per-module
+}
+
+void FGeometryFlowMeshProcessingEditorModule::ShutdownModule()
+{
+	// This function may be called during shutdown to clean up your module.  For modules that support dynamic reloading,
+	// we call this function before unloading the module.
+}
+
+#undef LOCTEXT_NAMESPACE
+	
+IMPLEMENT_MODULE(FGeometryFlowMeshProcessingEditorModule, GeometryFlowMeshProcessingEditor)
diff --git a/Engine/Plugins/Experimental/GeometryFlow/Source/GeometryFlowMeshProcessingEditor/Private/MeshProcessingNodes/MeshAutoGenerateUVsNode.cpp b/Engine/Plugins/Experimental/GeometryFlow/Source/GeometryFlowMeshProcessingEditor/Private/MeshProcessingNodes/MeshAutoGenerateUVsNode.cpp
new file mode 100644
--- /dev/null
+++ b/Engine/Plugins/Experimental/GeometryFlow/Source/GeometryFlowMeshProcessingEditor/Private/MeshProcessingNodes/MeshAutoGenerateUVsNode.cpp
@@ -0,0 +1,29 @@
+// Copyright Epic Games, Inc. All Rights Reserved.
+
+#include "MeshProcessingNodes/MeshAutoGenerateUVsNode.h"
+
+#include "ParameterizationOps/ParameterizeMeshOp.h"
+
+
+using namespace UE::GeometryFlow;
+
+
+void FMeshAutoGenerateUVsNode::GenerateUVs(const FDynamicMesh3& MeshIn, const FMeshAutoGenerateUVsSettings& Settings, FDynamicMesh3& MeshOut)
+{
+	// this is horrible - have to copy input mesh so that we can pass a TSharedPtr
+	TSharedPtr<FDynamicMesh3> InputMesh = MakeShared<FDynamicMesh3>(MeshIn);
+
+	FParameterizeMeshOp ParameterizeMeshOp;
+	ParameterizeMeshOp.Stretch = Settings.Stretch;
+	ParameterizeMeshOp.NumCharts = Settings.NumCharts;
+	ParameterizeMeshOp.InputMesh = InputMesh;
+
+	ParameterizeMeshOp.IslandMode = EParamOpIslandMode::Auto;
+	ParameterizeMeshOp.UnwrapType = EParamOpUnwrapType::MinStretch;
+
+	FProgressCancel Progress;
+	ParameterizeMeshOp.CalculateResult(&Progress);
+	TUniquePtr<FDynamicMesh3> ResultMesh = ParameterizeMeshOp.ExtractResult();
+
+	MeshOut = MoveTemp(*ResultMesh);
+}
\ No newline at end of file
diff --git a/Engine/Plugins/Experimental/GeometryFlow/Source/GeometryFlowMeshProcessingEditor/Public/GeometryFlowMeshProcessingEditorModule.h b/Engine/Plugins/Experimental/GeometryFlow/Source/GeometryFlowMeshProcessingEditor/Public/GeometryFlowMeshProcessingEditorModule.h
new file mode 100644
--- /dev/null
+++ b/Engine/Plugins/Experimental/GeometryFlow/Source/GeometryFlowMeshProcessingEditor/Public/GeometryFlowMeshProcessingEditorModule.h
@@ -0,0 +1,15 @@
+// Copyright Epic Games, Inc. All Rights Reserved.
+
+#pragma once
+
+#include "CoreMinimal.h"
+#include "Modules/ModuleManager.h"
+
+class FGeometryFlowMeshProcessingEditorModule : public IModuleInterface
+{
+public:
+
+	/** IModuleInterface implementation */
+	virtual void StartupModule() override;
+	virtual void ShutdownModule() override;
+};
diff --git a/Engine/Plugins/Experimental/GeometryFlow/Source/GeometryFlowMeshProcessingEditor/Public/MeshProcessingNodes/MeshAutoGenerateUVsNode.h b/Engine/Plugins/Experimental/GeometryFlow/Source/GeometryFlowMeshProcessingEditor/Public/MeshProcessingNodes/MeshAutoGenerateUVsNode.h
new file mode 100644
--- /dev/null
+++ b/Engine/Plugins/Experimental/GeometryFlow/Source/GeometryFlowMeshProcessingEditor/Public/MeshProcessingNodes/MeshAutoGenerateUVsNode.h
@@ -0,0 +1,51 @@
+// Copyright Epic Games, Inc. All Rights Reserved.
+
+#pragma once
+
+#include "MeshProcessingNodes/MeshProcessingBaseNodes.h"
+#include "MeshProcessingNodes/MeshProcessingDataTypesEditor.h"
+
+namespace UE
+{
+namespace GeometryFlow
+{
+
+
+struct GEOMETRYFLOWMESHPROCESSINGEDITOR_API FMeshAutoGenerateUVsSettings
+{
+	DECLARE_GEOMETRYFLOW_DATA_TYPE_IDENTIFIER(EMeshProcessingDataTypesEditor::MeshAutoGenerateUVsSettings);
+
+	double Stretch = 0.5;
+	int NumCharts = 0;
+
+};
+GEOMETRYFLOW_DECLARE_SETTINGS_TYPES(FMeshAutoGenerateUVsSettings, MeshAutoGenerateUVs);
+
+
+
+class GEOMETRYFLOWMESHPROCESSINGEDITOR_API FMeshAutoGenerateUVsNode : public TProcessMeshWithSettingsBaseNode<FMeshAutoGenerateUVsSettings>
+{
+public:
+	FMeshAutoGenerateUVsNode() : TProcessMeshWithSettingsBaseNode<FMeshAutoGenerateUVsSettings>()
+	{
+	}
+
+	virtual void ProcessMesh(
+		const FNamedDataMap& DatasIn,
+		const FMeshAutoGenerateUVsSettings& Settings,
+		const FDynamicMesh3& MeshIn,
+		FDynamicMesh3& MeshOut) override
+	{
+		GenerateUVs(MeshIn, Settings, MeshOut);
+	}
+
+	virtual void GenerateUVs(const FDynamicMesh3& MeshIn, const FMeshAutoGenerateUVsSettings& Settings, FDynamicMesh3& MeshOut);
+
+};
+
+
+
+
+
+}	// end namespace GeometryFlow
+}	// end namespace UE
\ No newline at end of file
diff --git a/Engine/Plugins/Experimental/GeometryFlow/Source/GeometryFlowMeshProcessingEditor/Public/MeshProcessingNodes/MeshProcessingDataTypesEditor.h b/Engine/Plugins/Experimental/GeometryFlow/Source/GeometryFlowMeshProcessingEditor/Public/MeshProcessingNodes/MeshProcessingDataTypesEditor.h
new file mode 100644
--- /dev/null
+++ b/Engine/Plugins/Experimental/GeometryFlow/Source/GeometryFlowMeshProcessingEditor/Public/MeshProcessingNodes/MeshProcessingDataTypesEditor.h
@@ -0,0 +1,21 @@
+// Copyright Epic Games, Inc. All Rights Reserved.
+
+#pragma once
+
+#include "CoreMinimal.h"
+#include "GeometryFlowTypes.h"
+
+namespace UE
+{
+namespace GeometryFlow
+{
+
+enum class EMeshProcessingDataTypesEditor
+{
+	MeshAutoGenerateUVsSettings = (int)EDataTypes::BaseMeshProcessingTypesEditor + 1
+};
+
+
+
+}	// end namespace GeometryFlow
+}	// en
\ No newline at end of file
diff --git a/Engine/Plugins/Experimental/MeshLODToolset/MeshLODToolset.uplugin b/Engine/Plugins/Experimental/MeshLODToolset/MeshLODToolset.uplugin
new file mode 100644
--- /dev/null
+++ b/Engine/Plugins/Experimental/MeshLODToolset/MeshLODToolset.uplugin
@@ -0,0 +1,43 @@
+{
+	"FileVersion": 3,
+	"Version": 1,
+	"VersionName": "0.1",
+	"FriendlyName": "Mesh LOD Toolset",
+	"Description": "A set of modules implementing 3D mesh LOD creation",
+	"Category": "Other",
+	"CreatedBy": "Epic Games, Inc.",
+	"CreatedByURL": "http://epicgames.com",
+	"DocsURL": "",
+	"MarketplaceURL": "",
+	"SupportURL": "",
+	"CanContainContent": true,
+	"IsBetaVersion": true,
+	"Installed": false,
+	"Hidden": true,
+	"Modules": [
+		{
+			"Name": "MeshLODToolset",
+			"Type": "Editor",
+			"LoadingPhase": "Default",
+			"WhitelistPlatforms": [ "Win64" ]
+		},
+	],
+	"Plugins": [
+		{
+			"Name": "GeometryProcessing",
+			"Enabled": true
+		},
+		{
+			"Name": "GeometryFlow",
+			"Enabled": true
+		},
+		{
+			"Name": "MeshModelingToolset",
+			"Enabled": true
+		},
+		{
+			"Name": "EditorScriptingUtilities",
+			"Enabled": true
+		}
+	]
+}
diff --git a/Engine/Plugins/Experimental/MeshLODToolset/Source/MeshLODToolset/MeshLODToolset.Build.cs b/Engine/Plugins/Experimental/MeshLODToolset/Source/MeshLODToolset/MeshLODToolset.Build.cs
new file mode 100644
--- /dev/null
+++ b/Engine/Plugins/Experimental/MeshLODToolset/Source/MeshLODToolset/MeshLODToolset.Build.cs
@@ -0,0 +1,87 @@
+// Copyright Epic Games, Inc. All Rights Reserved.
+
+using UnrealBuildTool;
+
+public class MeshLODToolset : ModuleRules
+{
+	public MeshLODToolset(ReadOnlyTargetRules Target) : base(Target)
+	{
+		PCHUsage = ModuleRules.PCHUsageMode.UseExplicitOrSharedPCHs;
+		
+		PublicIncludePaths.AddRange(
+			new string[] {
+				// ... add public include paths required here ...
+			}
+			);
+				
+		
+		PrivateIncludePaths.AddRange(
+			new string[] {
+				// ... add other private include paths required here ...
+			}
+			);
+			
+		
+		PublicDependencyModuleNames.AddRange(
+			new string[]
+			{
+				"Core",
+				"Eigen",
+                "InteractiveToolsFramework",
+				"GeometricObjects",
+				"GeometryAlgorithms",
+				"DynamicMesh",
+				"MeshConversion",
+				"MeshDescription",
+                "StaticMeshDescription",
+				"ModelingComponents",
+				"ModelingOperators",
+				"GeometryFlowCore",
+				"GeometryFlowMeshProcessing",
+				"GeometryFlowMeshProcessingEditor"
+
+				// ... add other public dependencies that you statically link with here ...
+			}
+            );
+			
+		
+		PrivateDependencyModuleNames.AddRange(
+			new string[]
+			{
+
+				//"RenderCore",
+				//"RHI",
+				//"MeshUtilities",    // temp for saving mesh asset
+				//"UnrealEd",
+				//"MeshBuilder",
+                //"MeshDescriptionOperations",
+				//"MeshUtilitiesCommon",
+				//"MeshReductionInterface", // for UE4 standard simplification
+                //"ProxyLODMeshReduction", // for mesh merging voxel-based csg
+				//"Slate",
+				//"SlateCore",
+
+				"CoreUObject",
+				"Engine",
+				"ModelingOperators",
+				"InputCore",
+				"PhysicsCore",
+				"RenderCore",
+
+				"UnrealEd",
+				"EditorScriptingUtilities",
+				"AssetTools"
+
+				// ... add private dependencies that you statically link with here ...
+			}
+			);
+		
+		
+		DynamicallyLoadedModuleNames.AddRange(
+			new string[]
+			{
+				// ... add any modules that your module loads dynamically here ...
+			}
+			);
+	}
+}
diff --git a/Engine/Plugins/Experimental/MeshLODToolset/Source/MeshLODToolset/Private/Graphs/GenerateMeshLODGraph.cpp b/Engine/Plugins/Experimental/MeshLODToolset/Source/MeshLODToolset/Private/Graphs/GenerateMeshLODGraph.cpp
new file mode 100644
--- /dev/null
+++ b/Engine/Plugins/Experimental/MeshLODToolset/Source/MeshLODToolset/Private/Graphs/GenerateMeshLODGraph.cpp
@@ -0,0 +1,290 @@
+// Copyright Epic Games, Inc. All Rights Reserved.
+
+#include "Graphs/GenerateMeshLODGraph.h"
+
+
+#include "GeometryFlowGraph.h"
+#include "GeometryFlowGraphUtil.h"
+#include "BaseNodes/TransferNode.h"
+
+#include "MeshProcessingNodes/MeshSolidifyNode.h"
+#include "MeshProcessingNodes/MeshVoxMorphologyNode.h"
+#include "MeshProcessingNodes/MeshSimplifyNode.h"
+#include "MeshProcessingNodes/MeshDeleteTrianglesNode.h"
+
+#include "MeshProcessingNodes/MeshNormalsNodes.h"
+#include "MeshProcessingNodes/MeshTangentsNodes.h"
+
+#include "MeshProcessingNodes/MeshAutoGenerateUVsNode.h"
+#include "MeshProcessingNodes/MeshRecalculateUVsNode.h"
+#include "MeshProcessingNodes/MeshRepackUVsNode.h"
+
+#include "DataTypes/MeshImageBakingData.h"
+#include "MeshBakingNodes/BakeMeshNormalMapNode.h"
+#include "MeshBakingNodes/BakeMeshTextureImageNode.h"
+
+#include "MeshDecompositionNodes/MakeTriangleSetsNode.h"
+#include "PhysicsNodes/GenerateConvexHullsCollisionNode.h"
+
+
+using namespace UE::GeometryFlow;
+
+
+
+void FGenerateMeshLODGraph::SetSourceMesh(const FDynamicMesh3& SourceMeshIn)
+{
+	UpdateSourceNodeValue<FDynamicMeshSourceNode>(*Graph, MeshSourceNode, SourceMeshIn);
+}
+
+
+void FGenerateMeshLODGraph::EvaluateResult(
+	FDynamicMesh3& ResultMesh,
+	FSimpleShapeSet3d& ResultCollision,
+	UE::GeometryFlow::FNormalMapImage& NormalMap,
+	TArray<TUniquePtr<UE::GeometryFlow::FTextureImage>>& TextureImages)
+{
+	//
+	// evaluate normal map
+	//
+
+	NormalMap = FNormalMapImage();
+	TUniquePtr<FEvaluationInfo> NormalMapEvalInfo = MakeUnique<FEvaluationInfo>();
+	EGeometryFlowResult NormalMapEvalResult = Graph->EvaluateResult(BakeNormalMapNode, FBakeMeshNormalMapNode::OutParamNormalMap(),
+		NormalMap, (int)EMeshProcessingDataTypes::NormalMapImage, NormalMapEvalInfo, true);
+	ensure(NormalMapEvalResult == EGeometryFlowResult::Ok);
+
+	UE_LOG(LogTemp, Warning, TEXT("NormalMapPass - Evaluated %d Nodes, Recomputed %d"), NormalMapEvalInfo->NumEvaluations(), NormalMapEvalInfo->NumComputes());
+
+	//
+	// evaluate transferred textures
+	//
+
+	for (FBakeTextureGraphInfo& TexBakeStep : BakeTextureNodes)
+	{
+		TUniquePtr<UE::GeometryFlow::FTextureImage> NewImage = MakeUnique<UE::GeometryFlow::FTextureImage>();
+		TUniquePtr<FEvaluationInfo> TexBakeEvalInfo = MakeUnique<FEvaluationInfo>();
+		EGeometryFlowResult TexBakeEvalResult = Graph->EvaluateResult(TexBakeStep.BakeNode, FBakeMeshTextureImageNode::OutParamTextureImage(),
+			*NewImage, (int)EMeshProcessingDataTypes::TextureImage, TexBakeEvalInfo, true);
+		TextureImages.Add(MoveTemp(NewImage));
+		ensure(TexBakeEvalResult == EGeometryFlowResult::Ok);
+
+		UE_LOG(LogTemp, Warning, TEXT("TextureBakePass %s - Evaluated %d Nodes, Recomputed %d"), *TexBakeStep.Identifier, TexBakeEvalInfo->NumEvaluations(), TexBakeEvalInfo->NumComputes());
+	}
+
+	//
+	// evaluate collision
+	//
+
+	bool bTakeResultCollision = false;
+	ResultCollision = FSimpleShapeSet3d();
+
+	TUniquePtr<FEvaluationInfo> CollisionEvalInfo = MakeUnique<FEvaluationInfo>();
+	EGeometryFlowResult CollisionEvalResult = Graph->EvaluateResult(CollisionOutputNode, FCollisionGeometryTransferNode::OutParamValue(),
+		ResultCollision, FCollisionGeometry::DataTypeIdentifier, CollisionEvalInfo, bTakeResultCollision);
+	ensure(CollisionEvalResult == EGeometryFlowResult::Ok);
+
+	UE_LOG(LogTemp, Warning, TEXT("OutputCollisionPass - Evaluated %d Nodes, Recomputed %d"), CollisionEvalInfo->NumEvaluations(), CollisionEvalInfo->NumComputes());
+
+
+	//
+	// evaluate result mesh
+	// 
+
+
+	bool bTakeResultMesh = true;
+	ResultMesh.Clear();
+
+	TUniquePtr<FEvaluationInfo> MeshEvalInfo = MakeUnique<FEvaluationInfo>();
+	EGeometryFlowResult EvalResult = Graph->EvaluateResult(MeshOutputNode, FDynamicMeshTransferNode::OutParamValue(),
+		ResultMesh, (int32)EMeshProcessingDataTypes::DynamicMesh, MeshEvalInfo, bTakeResultMesh);
+	ensure(EvalResult == EGeometryFlowResult::Ok);
+
+	UE_LOG(LogTemp, Warning, TEXT("OutputMeshPass - Evaluated %d Nodes, Recomputed %d"), MeshEvalInfo->NumEvaluations(), MeshEvalInfo->NumComputes());
+}
+
+
+
+
+
+void FGenerateMeshLODGraph::BuildGraph()
+{
+	Graph = MakeUnique<FGraph>();
+
+	MeshSourceNode = Graph->AddNodeOfType<FDynamicMeshSourceNode>(TEXT("SourceMesh"));
+
+	// generating low-poly mesh
+
+	SolidifyNode = Graph->AddNodeOfType<FSolidifyMeshNode>(TEXT("Solidify"));
+	Graph->InferConnection(MeshSourceNode, SolidifyNode);
+	FGraph::FHandle SolidifySettingsNode = Graph->AddNodeOfType<FSolidifySettingsSourceNode>(TEXT("SolidifySettings"));
+	Graph->InferConnection(SolidifySettingsNode, SolidifyNode);
+
+	MorphologyNode = Graph->AddNodeOfType<FVoxClosureMeshNode>(TEXT("Closure"));
+	Graph->InferConnection(SolidifyNode, MorphologyNode);
+	FGraph::FHandle MorphologySettingsNode = Graph->AddNodeOfType<FVoxClosureSettingsSourceNode>(TEXT("ClosureSettings"));
+	Graph->InferConnection(MorphologySettingsNode, MorphologyNode);
+
+	SimplifyNode = Graph->AddNodeOfType<FSimplifyMeshNode>(TEXT("Simplify"));
+	Graph->InferConnection(MorphologyNode, SimplifyNode);
+	FGraph::FHandle SimplifySettingsNode = Graph->AddNodeOfType<FSimplifySettingsSourceNode>(TEXT("SimplifySettings"));
+	Graph->InferConnection(SimplifySettingsNode, SimplifyNode);
+
+	NormalsNode = Graph->AddNodeOfType<FComputeMeshNormalsNode>(TEXT("Normals"));
+	Graph->InferConnection(SimplifyNode, NormalsNode);
+	FGraph::FHandle NormalsSettingsNode = Graph->AddNodeOfType<FNormalsSettingsSourceNode>(TEXT("NormalsSettings"));
+	Graph->InferConnection(NormalsSettingsNode, NormalsNode);
+
+	// computing UVs
+
+	AutoUVNode = Graph->AddNodeOfType<FMeshAutoGenerateUVsNode>(TEXT("AutoUV"));
+	Graph->InferConnection(NormalsNode, AutoUVNode);
+	FGraph::FHandle AutoUVSettingsNode = Graph->AddNodeOfType<FMeshAutoGenerateUVsSettingsSourceNode>(TEXT("AutoUVSettings"));
+	Graph->InferConnection(AutoUVSettingsNode, AutoUVNode);
+
+	RecomputeUVNode = Graph->AddNodeOfType<FMeshRecalculateUVsNode>(TEXT("RecalcUV"));
+	Graph->InferConnection(AutoUVNode, RecomputeUVNode);
+	FGraph::FHandle RecomputeUVSettingsNode = Graph->AddNodeOfType<FMeshRecalculateUVsSettingsSourceNode>(TEXT("RecalcUVSettings"));
+	Graph->InferConnection(RecomputeUVSettingsNode, RecomputeUVNode);
+
+	RepackUVNode = Graph->AddNodeOfType<FMeshRepackUVsNode>(TEXT("RepackUV"));
+	Graph->InferConnection(RecomputeUVNode, RepackUVNode);
+	FGraph::FHandle RepackUVSettingsNode = Graph->AddNodeOfType<FMeshRepackUVsSettingsSourceNode>(TEXT("RepackUVSettings"));
+	Graph->InferConnection(RepackUVSettingsNode, RepackUVNode);
+
+
+	// final mesh output
+
+	MeshOutputNode = Graph->AddNodeOfType<FDynamicMeshTransferNode>(TEXT("OutputMesh"));
+	Graph->InferConnection(RepackUVNode, MeshOutputNode);
+
+
+	// create tangents
+
+	TangentsNode = Graph->AddNodeOfType<FComputeMeshTangentsNode>(TEXT("Tangents"));
+	Graph->InferConnection(RepackUVNode, TangentsNode);
+	FGraph::FHandle TangentsSettingsNode = Graph->AddNodeOfType<FTangentsSettingsSourceNode>(TEXT("TangentsSettings"));
+	Graph->InferConnection(TangentsSettingsNode, TangentsNode);
+
+	// create bake cache
+
+	BakeCacheNode = Graph->AddNodeOfType<FMakeMeshBakingCacheNode>(TEXT("MakeBakeCache"));
+	Graph->AddConnection(MeshSourceNode, FDynamicMeshSourceNode::OutParamValue(), BakeCacheNode, FMakeMeshBakingCacheNode::InParamDetailMesh());
+	Graph->AddConnection(MeshOutputNode, FDynamicMeshTransferNode::OutParamValue(), BakeCacheNode, FMakeMeshBakingCacheNode::InParamTargetMesh());
+	FGraph::FHandle BakeCacheSettingsNode = Graph->AddNodeOfType<FMeshMakeBakingCacheSettingsSourceNode>(TEXT("BakeCacheSettings"));
+	Graph->InferConnection(BakeCacheSettingsNode, BakeCacheNode);
+
+	// normal map baker
+
+	BakeNormalMapNode = Graph->AddNodeOfType<FBakeMeshNormalMapNode>(TEXT("BakeNormalMap"));
+	Graph->InferConnection(BakeCacheNode, BakeNormalMapNode);
+	Graph->InferConnection(TangentsNode, BakeNormalMapNode);
+	FGraph::FHandle BakeNormalMapSettingsNode = Graph->AddNodeOfType<FBakeMeshNormalMapSettingsSourceNode>(TEXT("BakeNormalMapSettings"));
+	Graph->InferConnection(BakeNormalMapSettingsNode, BakeNormalMapNode);
+
+
+	Graph->AddNodeOfType<FMeshDeleteTrianglesNode>(TEXT("TestDeleteTrisNode"));
+
+	// collision generation
+
+	FGraph::FHandle IgnoreGroupsForCollisionNode = Graph->AddNodeOfType<FIndexSetsSourceNode>(TEXT("CollisionIgnoreGroups"));
+
+	//DecomposeMeshForCollisionNode = Graph->AddNodeOfType<FMakeTriangleSetsFromMeshNode>(TEXT("Decompose"));
+	DecomposeMeshForCollisionNode = Graph->AddNodeOfType<FMakeTriangleSetsFromGroupsNode>(TEXT("Decompose"));
+	Graph->InferConnection(MeshSourceNode, DecomposeMeshForCollisionNode);
+	Graph->InferConnection(IgnoreGroupsForCollisionNode, DecomposeMeshForCollisionNode);
+
+	GenerateConvexesNode = Graph->AddNodeOfType<FGenerateConvexHullsCollisionNode>(TEXT("GenerateConvexes"));
+	Graph->InferConnection(MeshSourceNode, GenerateConvexesNode);
+	Graph->InferConnection(DecomposeMeshForCollisionNode, GenerateConvexesNode);
+	FGraph::FHandle GenerateConvexesSettingsNode = Graph->AddNodeOfType<FGenerateConvexHullsCollisionSettingsSourceNode>(TEXT("GenerateConvexesSettings"));
+	Graph->InferConnection(GenerateConvexesSettingsNode, GenerateConvexesNode);
+
+	// final mesh output
+
+	CollisionOutputNode = Graph->AddNodeOfType<FCollisionGeometryTransferNode>(TEXT("OutputCollision"));
+	Graph->InferConnection(GenerateConvexesNode, CollisionOutputNode);
+
+
+	//
+	// parameters
+	//
+
+
+	FMeshSolidifySettings SolidifySettings;
+	UpdateSettingsSourceNodeValue(*Graph, SolidifySettingsNode, SolidifySettings);
+
+	FVoxClosureSettings MorphologySettings;
+	MorphologySettings.Distance = 5.0;
+	UpdateSettingsSourceNodeValue(*Graph, MorphologySettingsNode, MorphologySettings);
+
+	FMeshSimplifySettings SimplifySettings;
+	SimplifySettings.bDiscardAttributes = true;
+	SimplifySettings.SimplifyType = EMeshSimplifyType::VolumePreserving;
+	SimplifySettings.TargetType = EMeshSimplifyTargetType::TriangleCount;
+	SimplifySettings.TargetCount = 500;
+	UpdateSettingsSourceNodeValue(*Graph, SimplifySettingsNode, SimplifySettings);
+
+	FMeshNormalsSettings NormalsSettings;
+	NormalsSettings.NormalsType = EComputeNormalsType::FromFaceAngleThreshold;
+	NormalsSettings.AngleThresholdDeg = 45.0;
+	UpdateSettingsSourceNodeValue(*Graph, NormalsSettingsNode, NormalsSettings);
+
+	FMeshAutoGenerateUVsSettings AutoUVSettings;
+	UpdateSettingsSourceNodeValue(*Graph, AutoUVSettingsNode, AutoUVSettings);
+
+	FMeshRecalculateUVsSettings RecomputeUVSettings;
+	UpdateSettingsSourceNodeValue(*Graph, RecomputeUVSettingsNode, RecomputeUVSettings);
+
+	FMeshRepackUVsSettings RepackUVSettings;
+	UpdateSettingsSourceNodeValue(*Graph, RepackUVSettingsNode, RepackUVSettings);
+
+
+	FMeshTangentsSettings TangentsSettings;
+	UpdateSettingsSourceNodeValue(*Graph, TangentsSettingsNode, TangentsSettings);
+
+
+	FMeshMakeBakingCacheSettings BakeCacheSettings;
+	BakeCacheSettings.Dimensions = FImageDimensions(512, 512);
+	BakeCacheSettings.Thickness = 5.0;
+	UpdateSettingsSourceNodeValue(*Graph, BakeCacheSettingsNode, BakeCacheSettings);
+
+
+	FBakeMeshNormalMapSettings NormalMapSettings;
+	UpdateSettingsSourceNodeValue(*Graph, BakeNormalMapSettingsNode, NormalMapSettings);
+
+	FIndexSets IgnoreGroupsForCollision;
+	IgnoreGroupsForCollision.AppendSet({ 0 });
+	UpdateSettingsSourceNodeValue(*Graph, IgnoreGroupsForCollisionNode, IgnoreGroupsForCollision);
+
+	FGenerateConvexHullsCollisionSettings GenConvexesSettings;
+	UpdateSettingsSourceNodeValue(*Graph, GenerateConvexesSettingsNode, GenConvexesSettings);
+
+}
+
+
+
+
+int32 FGenerateMeshLODGraph::AppendTextureBakeNode(const TImageBuilder<FVector4f>& SourceImage, const FString& Identifier)
+{
+	FBakeTextureGraphInfo NewNode;
+	NewNode.Index = BakeTextureNodes.Num();
+	NewNode.Identifier = Identifier;
+
+	// add source node
+	NewNode.TexSourceNode = Graph->AddNodeOfType<FTextureImageSourceNode>(FString::Printf(TEXT("TextureSource%d_%s"), NewNode.Index, *NewNode.Identifier));
+
+	// normal map baker
+	NewNode.BakeNode = Graph->AddNodeOfType<FBakeMeshTextureImageNode>(FString::Printf(TEXT("BakeTexImage%d_%s"), NewNode.Index, *NewNode.Identifier));
+	ensure(Graph->InferConnection(BakeCacheNode, NewNode.BakeNode) == EGeometryFlowResult::Ok);
+	ensure(Graph->InferConnection(NewNode.TexSourceNode, NewNode.BakeNode) == EGeometryFlowResult::Ok);
+	FGraph::FHandle BakeTextureImageSettingsNode = Graph->AddNodeOfType<FBakeMeshTextureImageSettingsSourceNode>(TEXT("BakeTextureImageSettings"));
+	ensure(Graph->InferConnection(BakeTextureImageSettingsNode, NewNode.BakeNode) == EGeometryFlowResult::Ok);
+
+	FTextureImage InputTexImage;
+	InputTexImage.Image = SourceImage;
+	UpdateSourceNodeValue<FTextureImageSourceNode>(*Graph, NewNode.TexSourceNode, InputTexImage);
+
+	BakeTextureNodes.Add(NewNode);
+
+	return NewNode.Index;
+}
diff --git a/Engine/Plugins/Experimental/MeshLODToolset/Source/MeshLODToolset/Private/Graphs/GenerateStaticMeshLODProcess.cpp b/Engine/Plugins/Experimental/MeshLODToolset/Source/MeshLODToolset/Private/Graphs/GenerateStaticMeshLODProcess.cpp
new file mode 100644
--- /dev/null
+++ b/Engine/Plugins/Experimental/MeshLODToolset/Source/MeshLODToolset/Private/Graphs/GenerateStaticMeshLODProcess.cpp
@@ -0,0 +1,506 @@
+// Copyright Epic Games, Inc. All Rights Reserved.
+
+#include "Graphs/GenerateStaticMeshLODProcess.h"
+
+#include "Async/Async.h"
+#include "Async/ParallelFor.h"
+
+#include "MeshDescriptionToDynamicMesh.h"
+#include "DynamicMeshToMeshDescription.h"
+
+#include "AssetUtils/Texture2DUtil.h"
+#include "AssetUtils/Texture2DBuilder.h"
+
+#include "Physics/PhysicsDataCollection.h"
+
+#include "Misc/Paths.h"
+#include "EditorAssetLibrary.h"
+#include "AssetToolsModule.h"
+#include "AssetRegistryModule.h"
+#include "IAssetTools.h"
+#include "FileHelpers.h"
+
+#include "Materials/Material.h"
+#include "Materials/MaterialInstanceConstant.h"
+#include "Factories/MaterialInstanceConstantFactoryNew.h"
+
+#include "Engine/Classes/Engine/StaticMesh.h"
+#include "Engine/Classes/Components/StaticMeshComponent.h"
+#include "Engine/Classes/PhysicsEngine/BodySetup.h"
+
+
+
+
+namespace
+{
+#if WITH_EDITOR
+	static EAsyncExecution GenerateSMLODAsyncExecTarget = EAsyncExecution::LargeThreadPool;
+#else
+	static EAsyncExecution GenerateSMLODAsyncExecTarget = EAsyncExecution::ThreadPool;
+#endif
+}
+
+
+
+struct FReadTextureJob
+{
+	int32 MatIndex;
+	int32 TexIndex;
+};
+
+
+bool FGenerateStaticMeshLODProcess::Initialize(UStaticMesh* StaticMeshIn)
+{
+	if (!ensure(StaticMeshIn)) return false;
+	if (!ensure(StaticMeshIn->GetNumSourceModels() > 0)) return false;
+
+	// make sure we are not in rendering
+	FlushRenderingCommands();
+
+	SourceStaticMesh = StaticMeshIn;
+	const FMeshDescription* SourceMeshDescription = StaticMeshIn->GetMeshDescription(0);
+
+	// start async mesh-conversion
+	SourceMesh.Clear();
+	TFuture<void> ConvertMesh = Async(GenerateSMLODAsyncExecTarget, [&]()
+	{
+		FMeshDescriptionToDynamicMesh GetSourceMesh;
+		GetSourceMesh.Convert(SourceMeshDescription, SourceMesh);
+	});
+
+	// get list of source materials and find all the input texture params
+	const TArray<FStaticMaterial>& Materials = SourceStaticMesh->GetStaticMaterials();
+	SourceMaterials.SetNum(Materials.Num());
+	TArray<FReadTextureJob> ReadJobs;
+	for (int32 mi = 0; mi < Materials.Num(); ++mi)
+	{
+		SourceMaterials[mi].SourceMaterial = Materials[mi];
+
+		UMaterialInterface* MaterialInterface = Materials[mi].MaterialInterface;
+
+		TArray<FMaterialParameterInfo> ParameterInfo;
+		TArray<FGuid> ParameterIds;
+		MaterialInterface->GetAllTextureParameterInfo(ParameterInfo, ParameterIds);
+		for (int32 ti = 0; ti < ParameterInfo.Num(); ++ti)
+		{
+			FName ParamName = ParameterInfo[ti].Name;
+
+			UTexture* CurTexture = nullptr;
+			bool bFound = MaterialInterface->GetTextureParameterValue(
+				FMemoryImageMaterialParameterInfo(ParameterInfo[ti]), CurTexture);
+			if (ensure(bFound))
+			{
+				if (Cast<UTexture2D>(CurTexture) != nullptr)
+				{
+					FTextureInfo TexInfo;
+					TexInfo.SourceTexture = Cast<UTexture2D>(CurTexture);
+					TexInfo.ParameterName = ParamName;
+
+					TexInfo.bIsNormalMap = TexInfo.SourceTexture->IsNormalMap();
+					//TexInfo.bIsDefaultTexture = TexInfo.SourceTexture->IsDefaultTexture();
+					TexInfo.bIsDefaultTexture = UEditorAssetLibrary::GetPathNameForLoadedAsset(TexInfo.SourceTexture).StartsWith(TEXT("/Engine/"));
+
+					UE_LOG(LogTemp, Warning, TEXT("SOURCE TEX PATH IS %s - DefaultTex %d - SRGB %d"), *UEditorAssetLibrary::GetPathNameForLoadedAsset(TexInfo.SourceTexture),
+						TexInfo.bIsDefaultTexture?1:0, TexInfo.SourceTexture->SRGB);
+
+					TexInfo.bShouldBakeTexture = (TexInfo.bIsNormalMap == false && TexInfo.bIsDefaultTexture == false);
+
+					if (TexInfo.bShouldBakeTexture)
+					{
+						ReadJobs.Add(FReadTextureJob{ mi, SourceMaterials[mi].SourceTextures.Num() });
+					}
+
+					SourceMaterials[mi].SourceTextures.Add(TexInfo);
+				}
+			}
+		}
+	}
+
+
+	// extract all the texture params
+	TFuture<void> ReadTextures = Async(GenerateSMLODAsyncExecTarget, [&]()
+	{
+		ParallelFor(ReadJobs.Num(), [&](int32 ji)
+		{
+			FReadTextureJob job = ReadJobs[ji];
+			FTextureInfo& TexInfo = SourceMaterials[job.MatIndex].SourceTextures[job.TexIndex];
+			UE::AssetUtils::ReadTexture(TexInfo.SourceTexture, TexInfo.Dimensions, TexInfo.Image);
+		});
+	});
+	// single-thread path
+	//for (FReadTextureJob job : ReadJobs)
+	//{
+	//	FTextureInfo& TexInfo = SourceMaterials[job.MatIndex].SourceTextures[job.TexIndex];
+	//	UE::AssetUtils::ReadTexture(TexInfo.SourceTexture, TexInfo.Dimensions, TexInfo.Image);
+	//}
+
+
+	ConvertMesh.Wait();
+	ReadTextures.Wait();
+
+
+	FString FullPathWithExtension = UEditorAssetLibrary::GetPathNameForLoadedAsset(SourceStaticMesh);
+	SourceAssetPath = FPaths::GetBaseFilename(FullPathWithExtension, false);
+	SourceAssetFolder = FPaths::GetPath(SourceAssetPath);
+	SourceAssetName = FPaths::GetBaseFilename(FullPathWithExtension, true);
+
+	CalculateDerivedPathName(GetDefaultDerivedAssetSuffix());
+
+	return true;
+}
+
+
+
+
+void FGenerateStaticMeshLODProcess::CalculateDerivedPathName(FString NewAssetSuffix)
+{
+	DerivedSuffix = FPaths::MakeValidFileName(NewAssetSuffix);
+	if (DerivedSuffix.Len() == 0)
+	{
+		DerivedSuffix = GetDefaultDerivedAssetSuffix();
+	}
+	DerivedAssetPath = FString::Printf(TEXT("%s%s"), *GetSourceAssetPath(), *DerivedSuffix);
+	DerivedAssetFolder = FPaths::GetPath(DerivedAssetPath);
+}
+
+
+
+bool FGenerateStaticMeshLODProcess::ComputeDerivedSourceData()
+{
+	Generator = MakeUnique<FGenerateMeshLODGraph>();
+	Generator->BuildGraph();
+
+	TextureToDerivedTexIndex.Reset();
+	for (const FMaterialInfo& MatInfo : SourceMaterials)
+	{
+		for (const FTextureInfo& TexInfo : MatInfo.SourceTextures)
+		{
+			if (TexInfo.bShouldBakeTexture &&
+				(TextureToDerivedTexIndex.Contains(TexInfo.SourceTexture) == false) )
+			{
+				int32 NewIndex = Generator->AppendTextureBakeNode(TexInfo.Image, TexInfo.SourceTexture->GetName());
+				TextureToDerivedTexIndex.Add(TexInfo.SourceTexture, NewIndex);
+			}
+		}
+	}
+
+	Generator->SetSourceMesh(this->SourceMesh);
+
+	Generator->EvaluateResult(this->DerivedLODMesh, this->DerivedCollision, this->DerivedNormalMapImage, this->DerivedTextureImages);
+
+
+	// copy all materials for now...we are going to replace all the images though, and
+	// should not copy those?
+	DerivedMaterials = SourceMaterials;
+
+	// update texture data in derived materials
+	for (FMaterialInfo& MatInfo : DerivedMaterials)
+	{
+		for (FTextureInfo& TexInfo : MatInfo.SourceTextures)
+		{
+			if (TextureToDerivedTexIndex.Contains(TexInfo.SourceTexture))
+			{
+				int32 BakedTexIndex = TextureToDerivedTexIndex[TexInfo.SourceTexture];
+				const TUniquePtr<UE::GeometryFlow::FTextureImage>& DerivedTex = DerivedTextureImages[BakedTexIndex];
+				TexInfo.Dimensions = DerivedTex->Image.GetDimensions();
+				TexInfo.Image = MoveTemp(DerivedTex->Image);
+			}
+		}
+	}
+
+	return true;
+}
+
+
+
+bool FGenerateStaticMeshLODProcess::WriteDerivedAssetData()
+{
+	WriteDerivedTextures();
+
+	WriteDerivedMaterials();
+
+	WriteDerivedStaticMeshAsset();
+
+	return true;
+}
+
+
+
+void FGenerateStaticMeshLODProcess::WriteDerivedTextures()
+{
+	IAssetTools& AssetTools = FModuleManager::Get().LoadModuleChecked<FAssetToolsModule>("AssetTools").Get();
+
+	// write derived textures
+	int32 NumMaterials = SourceMaterials.Num();
+	check(DerivedMaterials.Num() == NumMaterials);
+	for (int32 mi = 0; mi < NumMaterials; ++mi)
+	{
+		const FMaterialInfo& SourceMaterialInfo = SourceMaterials[mi];
+		FMaterialInfo& DerivedMaterialInfo = DerivedMaterials[mi];
+
+		int32 NumTextures = SourceMaterialInfo.SourceTextures.Num();
+		check(DerivedMaterialInfo.SourceTextures.Num() == NumTextures);
+		for (int32 ti = 0; ti < NumTextures; ++ti)
+		{
+			const FTextureInfo& SourceTex = SourceMaterialInfo.SourceTextures[ti];
+			bool bConvertToSRGB = SourceTex.SourceTexture->SRGB;
+			FTextureInfo& DerivedTex = DerivedMaterialInfo.SourceTextures[ti];
+
+			if (DerivedTex.bShouldBakeTexture == false)
+			{
+				continue;
+			}
+
+			FTexture2DBuilder TextureBuilder;
+			TextureBuilder.Initialize(FTexture2DBuilder::ETextureType::Color, DerivedTex.Dimensions);
+			TextureBuilder.GetTexture2D()->SRGB = bConvertToSRGB;
+			TextureBuilder.Copy(DerivedTex.Image, bConvertToSRGB);
+			TextureBuilder.Commit(false);
+
+			DerivedTex.SourceTexture = TextureBuilder.GetTexture2D();
+			if (ensure(DerivedTex.SourceTexture))
+			{
+				FTexture2DBuilder::CopyPlatformDataToSourceData(DerivedTex.SourceTexture, FTexture2DBuilder::ETextureType::Color);
+
+				// write asset
+				bool bWriteOK = WriteDerivedTexture(SourceTex.SourceTexture, DerivedTex.SourceTexture);
+				ensure(bWriteOK);
+			}
+		}
+
+	}
+
+
+	// write derived normal map
+	{
+		FTexture2DBuilder NormapMapBuilder;
+		NormapMapBuilder.Initialize(FTexture2DBuilder::ETextureType::NormalMap, DerivedNormalMapImage.Image.GetDimensions());
+		NormapMapBuilder.Copy(DerivedNormalMapImage.Image, false);
+		NormapMapBuilder.Commit(false);
+
+		DerivedNormalMapTex = NormapMapBuilder.GetTexture2D();
+		if (ensure(DerivedNormalMapTex))
+		{
+			FTexture2DBuilder::CopyPlatformDataToSourceData(DerivedNormalMapTex, FTexture2DBuilder::ETextureType::NormalMap);
+
+			// write asset
+			bool bWriteOK = WriteDerivedTexture(DerivedNormalMapTex, SourceAssetName + TEXT("_NormalMap"));
+			ensure(bWriteOK);
+		}
+	}
+
+	
+}
+
+
+
+bool FGenerateStaticMeshLODProcess::WriteDerivedTexture(UTexture2D* SourceTexture, UTexture2D* DerivedTexture)
+{
+	IAssetTools& AssetTools = FModuleManager::Get().LoadModuleChecked<FAssetToolsModule>("AssetTools").Get();
+
+	FString SourceTexPath = UEditorAssetLibrary::GetPathNameForLoadedAsset(SourceTexture);
+	FString TexName = FPaths::GetBaseFilename(SourceTexPath, true);
+	return WriteDerivedTexture(DerivedTexture, TexName);
+}
+
+
+bool FGenerateStaticMeshLODProcess::WriteDerivedTexture(UTexture2D* DerivedTexture, FString BaseTexName)
+{
+	IAssetTools& AssetTools = FModuleManager::Get().LoadModuleChecked<FAssetToolsModule>("AssetTools").Get();
+
+	FString NewTexName = FString::Printf(TEXT("%s%s"), *BaseTexName, *DerivedSuffix);
+	FString NewAssetPath = FPaths::Combine(DerivedAssetFolder, NewTexName);
+
+	// delete existing asset so that we can have a clean duplicate
+	bool bNewAssetExists = UEditorAssetLibrary::DoesAssetExist(NewAssetPath);
+	if (bNewAssetExists)
+	{
+		bool bDeleteOK = UEditorAssetLibrary::DeleteAsset(NewAssetPath);
+		ensure(bDeleteOK);
+	}
+
+	// create package
+	FString UniquePackageName, UniqueAssetName;
+	AssetTools.CreateUniqueAssetName(NewAssetPath, TEXT(""), UniquePackageName, UniqueAssetName);
+	UPackage* AssetPackage = CreatePackage(*UniquePackageName);
+	check(AssetPackage);
+
+	// move texture from Transient package to new package
+	DerivedTexture->Rename(*UniqueAssetName, AssetPackage, REN_None);
+	// remove transient flag, add public/standalone/transactional
+	DerivedTexture->ClearFlags(RF_Transient);
+	DerivedTexture->SetFlags(RF_Public | RF_Standalone | RF_Transactional);
+	// do we need to Modify() it? we are not doing any undo/redo
+	DerivedTexture->Modify();
+	DerivedTexture->UpdateResource();
+	DerivedTexture->PostEditChange();		// this may be necessary if any Materials are using this texture
+	DerivedTexture->MarkPackageDirty();
+
+	FAssetRegistryModule::AssetCreated(DerivedTexture);		// necessary?
+
+	return true;
+}
+
+
+
+void FGenerateStaticMeshLODProcess::WriteDerivedMaterials()
+{
+	IAssetTools& AssetTools = FModuleManager::Get().LoadModuleChecked<FAssetToolsModule>("AssetTools").Get();
+
+	int32 NumMaterials = SourceMaterials.Num();
+	check(DerivedMaterials.Num() == NumMaterials);
+	for (int32 mi = 0; mi < NumMaterials; ++mi)
+	{
+		const FMaterialInfo& SourceMaterialInfo = SourceMaterials[mi];
+		FMaterialInfo& DerivedMaterialInfo = DerivedMaterials[mi];
+
+		UMaterialInterface* MaterialInterface = SourceMaterialInfo.SourceMaterial.MaterialInterface;
+		bool bSourceIsMIC = (Cast<UMaterialInstanceConstant>(MaterialInterface) != nullptr);
+
+		FString SourceMaterialPath = UEditorAssetLibrary::GetPathNameForLoadedAsset(MaterialInterface);
+		FString MaterialName = FPaths::GetBaseFilename(SourceMaterialPath, true);
+		FString NewMaterialName = FString::Printf(TEXT("%s%s"), *MaterialName, *DerivedSuffix);
+		FString NewMaterialPath = FPaths::Combine(DerivedAssetFolder, NewMaterialName);
+
+		// delete existing material so that we can have a clean duplicate
+		bool bNewMaterialExists = UEditorAssetLibrary::DoesAssetExist(NewMaterialPath);
+		if (bNewMaterialExists)
+		{
+			bool bDeleteOK = UEditorAssetLibrary::DeleteAsset(NewMaterialPath);
+			ensure(bDeleteOK);
+		}
+
+		// If source is a MIC, we can just duplicate it. If it is a UMaterial, we want to
+		// create a child MIC? Or we could dupe the Material and rewrite the textures.
+		// Probably needs to be an option.
+		UMaterialInstanceConstant* GeneratedMIC = nullptr;
+		if (bSourceIsMIC)
+		{
+			UObject* DupeAsset = UEditorAssetLibrary::DuplicateAsset(SourceMaterialPath, NewMaterialPath);
+			GeneratedMIC = Cast<UMaterialInstanceConstant>(DupeAsset);
+		}
+		else
+		{
+			UMaterial* SourceMaterial = MaterialInterface->GetBaseMaterial();
+			if (ensure(SourceMaterial))
+			{
+				UMaterialInstanceConstantFactoryNew* Factory = NewObject<UMaterialInstanceConstantFactoryNew>();
+				Factory->InitialParent = SourceMaterial;
+
+				UObject* NewAsset = AssetTools.CreateAsset(NewMaterialName, FPackageName::GetLongPackagePath(NewMaterialPath),
+					UMaterialInstanceConstant::StaticClass(), Factory);
+
+				GeneratedMIC = Cast<UMaterialInstanceConstant>(NewAsset);
+			}
+		}
+
+		// rewrite texture parameters to new textures
+		UpdateMaterialTextureParameters(GeneratedMIC, DerivedMaterialInfo);
+
+		// update StaticMaterial
+		DerivedMaterialInfo.SourceMaterial.MaterialInterface = GeneratedMIC;
+	}
+}
+
+
+
+void FGenerateStaticMeshLODProcess::UpdateMaterialTextureParameters(UMaterialInstanceConstant* Material, FMaterialInfo& DerivedMaterialInfo)
+{
+	Material->Modify();
+
+	int32 NumTextures = DerivedMaterialInfo.SourceTextures.Num();
+	for (int32 ti = 0; ti < NumTextures; ++ti)
+	{
+		FTextureInfo& DerivedTex = DerivedMaterialInfo.SourceTextures[ti];
+
+		if (DerivedTex.bIsNormalMap)
+		{
+			if (ensure(DerivedNormalMapTex))
+			{
+				FMaterialParameterInfo ParamInfo(DerivedTex.ParameterName);
+				Material->SetTextureParameterValueEditorOnly(ParamInfo, DerivedNormalMapTex);
+			}
+		}
+		else if (DerivedTex.bShouldBakeTexture)
+		{
+			UTexture2D* NewTexture = DerivedTex.SourceTexture;
+			if (ensure(NewTexture))
+			{
+				FMaterialParameterInfo ParamInfo(DerivedTex.ParameterName);
+				Material->SetTextureParameterValueEditorOnly(ParamInfo, NewTexture);
+			}
+		}
+	}
+
+	Material->PostEditChange();
+}
+
+
+
+void FGenerateStaticMeshLODProcess::WriteDerivedStaticMeshAsset()
+{
+	// [TODO] should we try to re-use existing asset here, or should we delete it? 
+	// The source asset might have had any number of config changes that we want to
+	// preserve in the duplicate...
+	UStaticMesh* GeneratedStaticMesh = nullptr;
+	if (UEditorAssetLibrary::DoesAssetExist(DerivedAssetPath))
+	{
+		UObject* LoadedAsset = UEditorAssetLibrary::LoadAsset(DerivedAssetPath);
+		GeneratedStaticMesh = Cast<UStaticMesh>(LoadedAsset);
+	}
+	else
+	{
+		UObject* DupeAsset = UEditorAssetLibrary::DuplicateAsset(SourceAssetPath, DerivedAssetPath);
+		GeneratedStaticMesh = Cast<UStaticMesh>(DupeAsset);
+	}
+
+	// make sure transactional flag is on
+	GeneratedStaticMesh->SetFlags(RF_Transactional);
+	GeneratedStaticMesh->Modify();
+
+	// update MeshDescription LOD0 mesh
+	GeneratedStaticMesh->SetNumSourceModels(1);
+	FMeshDescription* MeshDescription = GeneratedStaticMesh->GetMeshDescription(0);
+	FConversionToMeshDescriptionOptions ConversionOptions;
+	FDynamicMeshToMeshDescription Converter(ConversionOptions);
+	Converter.Convert(&DerivedLODMesh, *MeshDescription);
+	GeneratedStaticMesh->CommitMeshDescription(0);
+
+	TArray<FStaticMaterial> NewMaterials;
+	for (FMaterialInfo& DerivedMaterialInfo : DerivedMaterials)
+	{
+		NewMaterials.Add(DerivedMaterialInfo.SourceMaterial);
+	}
+
+	// update materials on generated mesh
+	GeneratedStaticMesh->SetStaticMaterials(NewMaterials);
+
+
+	// collision
+	FPhysicsDataCollection NewCollisionGeo;
+	NewCollisionGeo.Geometry = DerivedCollision;
+	NewCollisionGeo.CopyGeometryToAggregate();
+
+	// code below derived from FStaticMeshEditor::DuplicateSelectedPrims()
+	UBodySetup* BodySetup = GeneratedStaticMesh->GetBodySetup();
+	// mark the BodySetup for modification. Do we need to modify the UStaticMesh??
+	BodySetup->Modify();
+	//Clear the cache (PIE may have created some data), create new GUID    (comment from StaticMeshEditor)
+	BodySetup->InvalidatePhysicsData();
+	BodySetup->RemoveSimpleCollision();
+	BodySetup->AggGeom = NewCollisionGeo.AggGeom;
+	// update collision type
+	BodySetup->CollisionTraceFlag = ECollisionTraceFlag::CTF_UseDefault;
+	// rebuild physics data
+	BodySetup->InvalidatePhysicsData();
+	BodySetup->CreatePhysicsMeshes();
+
+	// do we need to do a post edit change here??
+
+	// is this necessary? 
+	GeneratedStaticMesh->CreateNavCollision(/*bIsUpdate=*/true);
+
+
+	// done updating mesh
+	GeneratedStaticMesh->PostEditChange();
+}
\ No newline at end of file
diff --git a/Engine/Plugins/Experimental/MeshLODToolset/Source/MeshLODToolset/Private/MeshLODToolsetModule.cpp b/Engine/Plugins/Experimental/MeshLODToolset/Source/MeshLODToolset/Private/MeshLODToolsetModule.cpp
new file mode 100644
--- /dev/null
+++ b/Engine/Plugins/Experimental/MeshLODToolset/Source/MeshLODToolset/Private/MeshLODToolsetModule.cpp
@@ -0,0 +1,20 @@
+// Copyright Epic Games, Inc. All Rights Reserved.
+
+#include "MeshLODToolsetModule.h"
+
+#define LOCTEXT_NAMESPACE "FMeshLODToolsetModule"
+
+void FMeshLODToolsetModule::StartupModule()
+{
+	// This code will execute after your module is loaded into memory; the exact timing is specified in the .uplugin file per-module
+}
+
+void FMeshLODToolsetModule::ShutdownModule()
+{
+	// This function may be called during shutdown to clean up your module.  For modules that support dynamic reloading,
+	// we call this function before unloading the module.
+}
+
+#undef LOCTEXT_NAMESPACE
+	
+IMPLEMENT_MODULE(FMeshLODToolsetModule, MeshLODToolset)
diff --git a/Engine/Plugins/Experimental/MeshLODToolset/Source/MeshLODToolset/Private/Tools/GenerateStaticMeshLODAssetTool.cpp b/Engine/Plugins/Experimental/MeshLODToolset/Source/MeshLODToolset/Private/Tools/GenerateStaticMeshLODAssetTool.cpp
new file mode 100644
--- /dev/null
+++ b/Engine/Plugins/Experimental/MeshLODToolset/Source/MeshLODToolset/Private/Tools/GenerateStaticMeshLODAssetTool.cpp
@@ -0,0 +1,405 @@
+// Copyright Epic Games, Inc. All Rights Reserved.
+
+#include "Tools/GenerateStaticMeshLODAssetTool.h"
+#include "InteractiveToolManager.h"
+#include "ToolBuilderUtil.h"
+
+#include "ToolSetupUtil.h"
+
+#include "Generators/SphereGenerator.h"
+
+#include "DynamicMesh3.h"
+#include "BaseBehaviors/MultiClickSequenceInputBehavior.h"
+#include "Selection/SelectClickedAction.h"
+#include "DynamicMeshEditor.h"
+#include "MeshTransforms.h"
+
+#include "MeshDescriptionToDynamicMesh.h"
+#include "DynamicMeshToMeshDescription.h"
+
+#include "AssetGenerationUtil.h"
+#include "Selection/ToolSelectionUtil.h"
+
+#include "Components/StaticMeshComponent.h"
+#include "Engine/StaticMesh.h"
+
+#include "Misc/Paths.h"
+#include "EditorAssetLibrary.h"
+#include "AssetToolsModule.h"
+#include "IAssetTools.h"
+
+#include "Materials/Material.h"
+#include "Materials/MaterialInstanceConstant.h"
+#include "Factories/MaterialInstanceConstantFactoryNew.h"
+
+#include "Graphs/GenerateStaticMeshLODProcess.h"
+
+
+#if WITH_EDITOR
+#include "Misc/ScopedSlowTask.h"
+#endif
+
+
+#define LOCTEXT_NAMESPACE "UGenerateStaticMeshLODAssetTool"
+
+
+/*
+ * ToolBuilder
+ */
+
+
+bool UGenerateStaticMeshLODAssetToolBuilder::CanBuildTool(const FToolBuilderState& SceneState) const
+{
+	// hack to make multi-tool look like single-tool
+	return (AssetAPI != nullptr && ToolBuilderUtil::CountComponents(SceneState, CanMakeComponentTarget) == 1);
+}
+
+UInteractiveTool* UGenerateStaticMeshLODAssetToolBuilder::BuildTool(const FToolBuilderState& SceneState) const
+{
+	UGenerateStaticMeshLODAssetTool* NewTool = NewObject<UGenerateStaticMeshLODAssetTool>(SceneState.ToolManager);
+
+	TArray<UActorComponent*> Components = ToolBuilderUtil::FindAllComponents(SceneState, CanMakeComponentTarget);
+	check(Components.Num() > 0);
+
+	TArray<TUniquePtr<FPrimitiveComponentTarget>> ComponentTargets;
+	for (UActorComponent* ActorComponent : Components)
+	{
+		auto* MeshComponent = Cast<UPrimitiveComponent>(ActorComponent);
+		if (MeshComponent)
+		{
+			ComponentTargets.Add(MakeComponentTarget(MeshComponent));
+		}
+	}
+
+	NewTool->SetSelection(MoveTemp(ComponentTargets));
+	NewTool->SetWorld(SceneState.World);
+	NewTool->SetAssetAPI(AssetAPI);
+
+	return NewTool;
+}
+
+
+
+
+/*
+ * Tool
+ */
+
+
+void UGenerateStaticMeshLODAssetTool::SetWorld(UWorld* World)
+{
+	this->TargetWorld = World;
+}
+
+
+void UGenerateStaticMeshLODAssetTool::Setup()
+{
+	UInteractiveTool::Setup();
+
+	BasicProperties = NewObject<UGenerateStaticMeshLODAssetToolProperties>(this);
+	AddToolPropertySource(BasicProperties);
+	BasicProperties->RestoreProperties(this);
+
+	BasicProperties->OutputName = AssetGenerationUtil::GetComponentAssetBaseName(ComponentTargets[0]->GetOwnerComponent());
+
+	BasicProperties->GeneratedSuffix = TEXT("_AutoLOD");
+
+
+	GetToolManager()->DisplayMessage(
+		LOCTEXT("OnStartStaticMeshLODAssetTool", "This tool creates a new LOD asset"),
+		EToolMessageLevel::UserNotification);
+
+
+
+	GenerateProcess = MakePimpl<FGenerateStaticMeshLODProcess>();
+
+	TUniquePtr<FPrimitiveComponentTarget>& SourceComponent = ComponentTargets[0];
+	UStaticMeshComponent* StaticMeshComponent = CastChecked<UStaticMeshComponent>(SourceComponent->GetOwnerComponent());
+	UStaticMesh* StaticMesh = StaticMeshComponent->GetStaticMesh();
+	if (StaticMesh)
+	{
+		GenerateProcess->Initialize(StaticMesh);
+	}
+
+}
+
+
+
+void UGenerateStaticMeshLODAssetTool::Shutdown(EToolShutdownType ShutdownType)
+{
+	BasicProperties->SaveProperties(this);
+
+	if (ShutdownType == EToolShutdownType::Accept)
+	{
+		CreateNewAsset();
+		//UpdateExistingAsset();
+	}
+}
+
+void UGenerateStaticMeshLODAssetTool::SetAssetAPI(IAssetGenerationAPI* AssetAPIIn)
+{
+	this->AssetAPI = AssetAPIIn;
+}
+
+bool UGenerateStaticMeshLODAssetTool::HasAccept() const
+{
+	return true;
+}
+
+bool UGenerateStaticMeshLODAssetTool::CanAccept() const
+{
+	return true;
+}
+
+
+void UGenerateStaticMeshLODAssetTool::CreateNewAsset()
+{
+	IAssetTools& AssetTools = FModuleManager::Get().LoadModuleChecked<FAssetToolsModule>("AssetTools").Get();
+
+	//GetToolManager()->BeginUndoTransaction( LOCTEXT("GenerateStaticMeshLODAssetTransaction", "Generate LOD Asset") );
+
+	GenerateProcess->CalculateDerivedPathName(BasicProperties->GeneratedSuffix);
+	GenerateProcess->ComputeDerivedSourceData();
+	GenerateProcess->WriteDerivedAssetData();
+
+/*
+	UStaticMesh* StaticMesh = GenerateProcess->GetSourceStaticMesh();
+	const FDynamicMesh3& GeneratedLOD0Mesh = GenerateProcess->GetDerivedLOD0Mesh();
+
+	// find existing path and generate AutoLOD path
+	FString FullSMPathWithExtension = UEditorAssetLibrary::GetPathNameForLoadedAsset(StaticMesh);
+	FString SMPathWithName = FPaths::GetBaseFilename(FullSMPathWithExtension, false);
+	FString SMPath = FPaths::GetPath(SMPathWithName);
+	FString UseSuffix = FPaths::MakeValidFileName(BasicProperties->GeneratedSuffix);
+	if (UseSuffix.Len() == 0)
+	{
+		UseSuffix = TEXT("_AutoLOD");
+	}
+	FString GeneratedSMPath = FString::Printf(TEXT("%s%s"), *SMPathWithName, *UseSuffix);
+
+	// [TODO] should we try to re-use existing asset here, or should we delete it? 
+	// The source asset might have had any number of config changes that we want to
+	// preserve in the duplicate...
+	UStaticMesh* GeneratedStaticMesh = nullptr;
+	if (UEditorAssetLibrary::DoesAssetExist(GeneratedSMPath))
+	{
+		UObject* LoadedAsset = UEditorAssetLibrary::LoadAsset(GeneratedSMPath);
+		GeneratedStaticMesh = Cast<UStaticMesh>(LoadedAsset);
+	}
+	else
+	{
+		UObject* DupeAsset = UEditorAssetLibrary::DuplicateAsset(SMPathWithName, GeneratedSMPath);
+		GeneratedStaticMesh = Cast<UStaticMesh>(DupeAsset);
+	}
+
+	// make sure transactional flag is on
+	GeneratedStaticMesh->SetFlags(RF_Transactional);
+	GeneratedStaticMesh->Modify();
+
+	// update MeshDescription LOD0 mesh
+	GeneratedStaticMesh->SetNumSourceModels(1);
+	FMeshDescription* MeshDescription = GeneratedStaticMesh->GetMeshDescription(0);
+	FConversionToMeshDescriptionOptions ConversionOptions;
+	FDynamicMeshToMeshDescription Converter(ConversionOptions);
+	Converter.Convert(&GeneratedLOD0Mesh, *MeshDescription);
+	GeneratedStaticMesh->CommitMeshDescription(0);
+
+	// done updating mesh
+	GeneratedStaticMesh->PostEditChange();
+
+
+
+	const TArray<FStaticMaterial>& OldMaterials = StaticMesh->GetStaticMaterials();
+	TArray<FStaticMaterial>& NewMaterials = GeneratedStaticMesh->GetStaticMaterials();
+	check(OldMaterials.Num() == NewMaterials.Num());
+
+	for (int32 mi = 0; mi < OldMaterials.Num(); ++mi)
+	{
+		UMaterialInterface* MaterialInterface = OldMaterials[mi].MaterialInterface;
+		bool bSourceIsMIC = (Cast<UMaterialInstanceConstant>(MaterialInterface) != nullptr);
+
+		FString SourceMaterialPath = UEditorAssetLibrary::GetPathNameForLoadedAsset(MaterialInterface);
+		FString MaterialName = FPaths::GetBaseFilename(SourceMaterialPath, true);
+		FString NewMaterialName = FString::Printf(TEXT("%s%s"), *MaterialName, *UseSuffix);
+		FString NewMaterialPath = FPaths::Combine(SMPath, NewMaterialName);
+
+
+		// delete existing material so that we can have a clean duplicate
+		bool bNewMaterialExists = UEditorAssetLibrary::DoesAssetExist(NewMaterialPath);
+		if (bNewMaterialExists)
+		{
+			bool bDeleteOK = UEditorAssetLibrary::DeleteAsset(NewMaterialPath);
+			ensure(bDeleteOK);
+		}
+
+		// If source is a MIC, we can just duplicate it. If it is a UMaterial, we want to
+		// create a child MIC? Or we could dupe the Material and rewrite the textures.
+		// Probably needs to be an option.
+		UMaterialInstanceConstant* GeneratedMIC = nullptr;
+		if (bSourceIsMIC)
+		{
+			UObject* DupeAsset = UEditorAssetLibrary::DuplicateAsset(SourceMaterialPath, NewMaterialPath);
+			GeneratedMIC = Cast<UMaterialInstanceConstant>(DupeAsset);
+		}
+		else
+		{
+			UMaterial* SourceMaterial = MaterialInterface->GetBaseMaterial();
+			if (ensure(SourceMaterial))
+			{
+				UMaterialInstanceConstantFactoryNew* Factory = NewObject<UMaterialInstanceConstantFactoryNew>();
+				Factory->InitialParent = SourceMaterial;
+
+				UObject* NewAsset = AssetTools.CreateAsset(NewMaterialName, FPackageName::GetLongPackagePath(NewMaterialPath),
+					UMaterialInstanceConstant::StaticClass(), Factory);
+
+				GeneratedMIC = Cast<UMaterialInstanceConstant>(NewAsset);
+			}
+		}
+
+		NewMaterials[mi].MaterialInterface = GeneratedMIC;
+	}
+
+
+	// update materials on generated mesh
+	GeneratedStaticMesh->SetStaticMaterials(NewMaterials);
+
+*/
+
+	// I think this would let us jump to new assets...
+	//ObjectsToSync.Add(NewAsset);
+	//GEditor->SyncBrowserToObjects(ObjectsToSync);
+
+
+	//GetToolManager()->EndUndoTransaction();
+}
+
+
+
+
+
+
+
+void UGenerateStaticMeshLODAssetTool::UpdateExistingAsset()
+{
+	GetToolManager()->BeginUndoTransaction(LOCTEXT("GenerateStaticMeshLODAssetToolTransactionName", "Combine Meshes"));
+
+	// note there is a MergeMeshUtilities.h w/ a very feature-filled mesh merging class, but for simplicity (and to fit modeling tool needs)
+	// this tool currently converts things through dynamic mesh instead
+
+	AActor* SkipActor = nullptr;
+
+#if WITH_EDITOR
+
+	bool bMergeSameMaterials = true;
+	TArray<UMaterialInterface*> AllMaterials;
+	TMap<UMaterialInterface*, int> KnownMaterials;
+	TArray<int> CombinedMatToOutMatIdx;
+	for (int32 ComponentIdx = 0; ComponentIdx < ComponentTargets.Num(); ComponentIdx++)
+	{
+		TUniquePtr<FPrimitiveComponentTarget>& ComponentTarget = ComponentTargets[ComponentIdx];
+		for (int MaterialIdx = 0, NumMaterials = ComponentTarget->GetNumMaterials(); MaterialIdx < NumMaterials; MaterialIdx++)
+		{
+			UMaterialInterface* Mat = ComponentTarget->GetMaterial(MaterialIdx);
+			int32 OutMatIdx = -1;
+			if (!bMergeSameMaterials || !KnownMaterials.Contains(Mat))
+			{
+				OutMatIdx = AllMaterials.Num();
+				if (bMergeSameMaterials)
+				{
+					KnownMaterials.Add(Mat, AllMaterials.Num());
+				}
+				AllMaterials.Add(Mat);
+			}
+			else
+			{
+				OutMatIdx = KnownMaterials[Mat];
+			}
+			CombinedMatToOutMatIdx.Add(OutMatIdx);
+		}
+	}
+
+	FDynamicMesh3 AccumulateDMesh;
+	AccumulateDMesh.EnableTriangleGroups();
+	AccumulateDMesh.EnableAttributes();
+	AccumulateDMesh.Attributes()->EnableMaterialID();
+
+	int32 SkipIndex = 0;
+	TUniquePtr<FPrimitiveComponentTarget>& UpdateTarget = ComponentTargets[SkipIndex];
+	SkipActor = UpdateTarget->GetOwnerActor();
+
+	FTransform3d TargetToWorld = (FTransform3d)UpdateTarget->GetWorldTransform();
+	FTransform3d WorldToTarget = TargetToWorld.Inverse();
+
+	{
+		FScopedSlowTask SlowTask(ComponentTargets.Num() + 1,
+			LOCTEXT("GenerateStaticMeshLODAssetBuild", "Building combined mesh ..."));
+		SlowTask.MakeDialog();
+
+		int MatIndexBase = 0;
+		for (int32 ComponentIdx = 0; ComponentIdx < ComponentTargets.Num(); ComponentIdx++)
+		{
+			SlowTask.EnterProgressFrame(1);
+			TUniquePtr<FPrimitiveComponentTarget>& ComponentTarget = ComponentTargets[ComponentIdx];
+
+			FMeshDescriptionToDynamicMesh Converter;
+			FDynamicMesh3 ComponentDMesh;
+			Converter.Convert(ComponentTarget->GetMesh(), ComponentDMesh);
+
+			// update material IDs to account for combined material set
+			FDynamicMeshMaterialAttribute* MatAttrib = ComponentDMesh.Attributes()->GetMaterialID();
+			for (int TID : ComponentDMesh.TriangleIndicesItr())
+			{
+				MatAttrib->SetValue(TID, CombinedMatToOutMatIdx[MatIndexBase + MatAttrib->GetValue(TID)]);
+			}
+			MatIndexBase += ComponentTarget->GetNumMaterials();
+
+
+			if (ComponentIdx != SkipIndex)
+			{
+				FTransform3d ComponentToWorld = (FTransform3d)ComponentTarget->GetWorldTransform();
+				MeshTransforms::ApplyTransform(ComponentDMesh, ComponentToWorld);
+				if (ComponentToWorld.GetDeterminant() < 0)
+				{
+					ComponentDMesh.ReverseOrientation(true);
+				}
+				MeshTransforms::ApplyTransform(ComponentDMesh, WorldToTarget);
+				if (WorldToTarget.GetDeterminant() < 0)
+				{
+					ComponentDMesh.ReverseOrientation(true);
+				}
+			}
+
+			FDynamicMeshEditor Editor(&AccumulateDMesh);
+			FMeshIndexMappings IndexMapping;
+			Editor.AppendMesh(&ComponentDMesh, IndexMapping);
+		}
+
+		SlowTask.EnterProgressFrame(1);
+
+		UpdateTarget->CommitMesh([&](const FPrimitiveComponentTarget::FCommitParams& CommitParams)
+		{
+			FDynamicMeshToMeshDescription Converter;
+			Converter.Convert(&AccumulateDMesh, *CommitParams.MeshDescription);
+		});
+
+		FComponentMaterialSet MaterialSet;
+		MaterialSet.Materials = AllMaterials;
+		UpdateTarget->CommitMaterialSetUpdate(MaterialSet, true);
+
+		// select the new actor
+		ToolSelectionUtil::SetNewActorSelection(GetToolManager(), SkipActor);
+	}
+#endif
+
+
+
+	GetToolManager()->EndUndoTransaction();
+}
+
+
+
+
+
+
+
+#undef LOCTEXT_NAMESPACE
diff --git a/Engine/Plugins/Experimental/MeshLODToolset/Source/MeshLODToolset/Public/Graphs/GenerateMeshLODGraph.h b/Engine/Plugins/Experimental/MeshLODToolset/Source/MeshLODToolset/Public/Graphs/GenerateMeshLODGraph.h
new file mode 100644
--- /dev/null
+++ b/Engine/Plugins/Experimental/MeshLODToolset/Source/MeshLODToolset/Public/Graphs/GenerateMeshLODGraph.h
@@ -0,0 +1,66 @@
+// Copyright Epic Games, Inc. All Rights Reserved.
+
+#pragma once
+
+#include "GeometryFlowGraph.h"
+#include "DataTypes/NormalMapData.h"
+#include "DataTypes/TextureImageData.h"
+#include "DataTypes/CollisionGeometryData.h"
+#include "DynamicMesh3.h"
+
+
+class FGenerateMeshLODGraph
+{
+public:
+	void BuildGraph();
+
+	int32 AppendTextureBakeNode(const TImageBuilder<FVector4f>& SourceImage, const FString& Identifier);
+
+	void SetSourceMesh(const FDynamicMesh3& SourceMesh);
+
+
+	void EvaluateResult(
+		FDynamicMesh3& ResultMesh,
+		FSimpleShapeSet3d& ResultCollision,
+		UE::GeometryFlow::FNormalMapImage& NormalMap,
+		TArray<TUniquePtr<UE::GeometryFlow::FTextureImage>>& TextureImages);
+
+
+protected:
+	TUniquePtr<UE::GeometryFlow::FGraph> Graph;
+
+
+	UE::GeometryFlow::FGraph::FHandle MeshSourceNode;
+
+	UE::GeometryFlow::FGraph::FHandle SolidifyNode;
+	UE::GeometryFlow::FGraph::FHandle MorphologyNode;
+
+	UE::GeometryFlow::FGraph::FHandle SimplifyNode;
+
+	UE::GeometryFlow::FGraph::FHandle NormalsNode;
+
+	UE::GeometryFlow::FGraph::FHandle AutoUVNode;
+	UE::GeometryFlow::FGraph::FHandle RecomputeUVNode;
+	UE::GeometryFlow::FGraph::FHandle RepackUVNode;
+
+	UE::GeometryFlow::FGraph::FHandle TangentsNode;
+
+	UE::GeometryFlow::FGraph::FHandle BakeCacheNode;
+	UE::GeometryFlow::FGraph::FHandle BakeNormalMapNode;
+
+	struct FBakeTextureGraphInfo
+	{
+		int32 Index;
+		FString Identifier;
+		UE::GeometryFlow::FGraph::FHandle TexSourceNode;
+		UE::GeometryFlow::FGraph::FHandle BakeNode;
+	};
+	TArray<FBakeTextureGraphInfo> BakeTextureNodes;
+
+	UE::GeometryFlow::FGraph::FHandle DecomposeMeshForCollisionNode;
+	UE::GeometryFlow::FGraph::FHandle GenerateConvexesNode;
+
+	UE::GeometryFlow::FGraph::FHandle CollisionOutputNode;
+	UE::GeometryFlow::FGraph::FHandle MeshOutputNode;
+
+};
\ No newline at end of file
diff --git a/Engine/Plugins/Experimental/MeshLODToolset/Source/MeshLODToolset/Public/Graphs/GenerateStaticMeshLODProcess.h b/Engine/Plugins/Experimental/MeshLODToolset/Source/MeshLODToolset/Public/Graphs/GenerateStaticMeshLODProcess.h
new file mode 100644
--- /dev/null
+++ b/Engine/Plugins/Experimental/MeshLODToolset/Source/MeshLODToolset/Public/Graphs/GenerateStaticMeshLODProcess.h
@@ -0,0 +1,92 @@
+// Copyright Epic Games, Inc. All Rights Reserved.
+
+#pragma once
+
+#include "Engine/StaticMesh.h"
+
+#include "DynamicMesh3.h"
+#include "Image/ImageBuilder.h"
+
+#include "Graphs/GenerateMeshLODGraph.h"
+
+class UTexture2D;
+class UMaterialInstanceConstant;
+
+class FGenerateStaticMeshLODProcess
+{
+public:
+
+	bool Initialize(UStaticMesh* SourceMesh);
+
+
+	UStaticMesh* GetSourceStaticMesh() const { return SourceStaticMesh; }
+	const FDynamicMesh3& GetSourceMesh() const { return SourceMesh; }
+
+	const FString& GetSourceAssetPath() const { return SourceAssetPath; }
+	const FString& GetSourceAssetFolder() const { return SourceAssetFolder; }
+	const FString& GetSourceAssetName() const { return SourceAssetName; }
+	static FString GetDefaultDerivedAssetSuffix() { return TEXT("_AutoLOD"); }
+
+	void CalculateDerivedPathName(FString NewAssetSuffix);
+
+
+	bool ComputeDerivedSourceData();
+	const FDynamicMesh3& GetDerivedLOD0Mesh() const { return DerivedLODMesh; }
+
+
+	bool WriteDerivedAssetData();
+
+
+
+protected:
+
+	UStaticMesh* SourceStaticMesh;
+	FString SourceAssetPath;
+	FString SourceAssetFolder;
+	FString SourceAssetName;
+
+	FDynamicMesh3 SourceMesh;
+
+	struct FTextureInfo
+	{
+		UTexture2D* SourceTexture;
+		FName ParameterName;
+		FImageDimensions Dimensions;
+		TImageBuilder<FVector4f> Image;
+		bool bIsNormalMap = false;
+		bool bIsDefaultTexture = false;
+		bool bShouldBakeTexture = false;
+	};
+
+	struct FMaterialInfo
+	{
+		FStaticMaterial SourceMaterial;
+		TArray<FTextureInfo> SourceTextures;
+	};
+
+	TArray<FMaterialInfo> SourceMaterials;
+
+
+
+	FString DerivedSuffix;
+	FString DerivedAssetFolder;
+	FString DerivedAssetPath;
+
+	FDynamicMesh3 DerivedLODMesh;
+	FSimpleShapeSet3d DerivedCollision;
+	UE::GeometryFlow::FNormalMapImage DerivedNormalMapImage;
+	TArray<TUniquePtr<UE::GeometryFlow::FTextureImage>> DerivedTextureImages;
+	TMap<UTexture2D*, int32> TextureToDerivedTexIndex;
+
+	TArray<FMaterialInfo> DerivedMaterials;
+	UTexture2D* DerivedNormalMapTex;
+
+	TUniquePtr<FGenerateMeshLODGraph> Generator;
+
+	bool WriteDerivedTexture(UTexture2D* SourceTexture, UTexture2D* DerivedTexture);
+	bool WriteDerivedTexture(UTexture2D* DerivedTexture, FString BaseTexName);
+	void WriteDerivedTextures();
+	void WriteDerivedMaterials();
+	void UpdateMaterialTextureParameters(UMaterialInstanceConstant* Material, FMaterialInfo& DerivedMaterialInfo);
+	void WriteDerivedStaticMeshAsset();
+};
\ No newline at end of file
diff --git a/Engine/Plugins/Experimental/MeshLODToolset/Source/MeshLODToolset/Public/MeshLODToolsetModule.h b/Engine/Plugins/Experimental/MeshLODToolset/Source/MeshLODToolset/Public/MeshLODToolsetModule.h
new file mode 100644
--- /dev/null
+++ b/Engine/Plugins/Experimental/MeshLODToolset/Source/MeshLODToolset/Public/MeshLODToolsetModule.h
@@ -0,0 +1,15 @@
+// Copyright Epic Games, Inc. All Rights Reserved.
+
+#pragma once
+
+#include "CoreMinimal.h"
+#include "Modules/ModuleManager.h"
+
+class FMeshLODToolsetModule : public IModuleInterface
+{
+public:
+
+	/** IModuleInterface implementation */
+	virtual void StartupModule() override;
+	virtual void ShutdownModule() override;
+};
diff --git a/Engine/Plugins/Experimental/MeshLODToolset/Source/MeshLODToolset/Public/Tools/GenerateStaticMeshLODAssetTool.h b/Engine/Plugins/Experimental/MeshLODToolset/Source/MeshLODToolset/Public/Tools/GenerateStaticMeshLODAssetTool.h
new file mode 100644
--- /dev/null
+++ b/Engine/Plugins/Experimental/MeshLODToolset/Source/MeshLODToolset/Public/Tools/GenerateStaticMeshLODAssetTool.h
@@ -0,0 +1,96 @@
+// Copyright Epic Games, Inc. All Rights Reserved.
+
+#pragma once
+
+#include "CoreMinimal.h"
+#include "UObject/NoExportTypes.h"
+#include "MultiSelectionTool.h"
+#include "InteractiveToolBuilder.h"
+#include "DynamicMesh3.h"
+#include "Templates/PimplPtr.h"
+#include "GenerateStaticMeshLODAssetTool.generated.h"
+
+
+// predeclarations
+struct FMeshDescription;
+class USimpleDynamicMeshComponent;
+class IAssetGenerationAPI;
+class FGenerateStaticMeshLODProcess;
+
+
+
+/**
+ *
+ */
+UCLASS()
+class MESHLODTOOLSET_API UGenerateStaticMeshLODAssetToolBuilder : public UInteractiveToolBuilder
+{
+	GENERATED_BODY()
+
+public:
+	IAssetGenerationAPI* AssetAPI = nullptr;
+
+	virtual bool CanBuildTool(const FToolBuilderState& SceneState) const override;
+	virtual UInteractiveTool* BuildTool(const FToolBuilderState& SceneState) const override;
+};
+
+
+/**
+ * Standard properties
+ */
+UCLASS()
+class MESHLODTOOLSET_API UGenerateStaticMeshLODAssetToolProperties : public UInteractiveToolPropertySet
+{
+	GENERATED_BODY()
+
+public:
+	UPROPERTY(EditAnywhere, Category = AssetOptions, meta = (TransientToolProperty))
+	FString OutputName;
+
+	/** Base name for newly-generated asset */
+	UPROPERTY(EditAnywhere, Category = AssetOptions, meta = (TransientToolProperty))
+	FString GeneratedSuffix;
+
+	/** Name of asset that will be updated */
+	UPROPERTY(VisibleAnywhere, Category = AssetOptions, meta = (TransientToolProperty))
+	FString OutputAsset;
+};
+
+
+
+
+
+
+/**
+ * Simple tool to combine multiple meshes into a single mesh asset
+ */
+UCLASS()
+class MESHLODTOOLSET_API UGenerateStaticMeshLODAssetTool : public UMultiSelectionTool
+{
+	GENERATED_BODY()
+
+public:
+	virtual void Setup() override;
+	virtual void Shutdown(EToolShutdownType ShutdownType) override;
+
+	virtual void SetWorld(UWorld* World);
+	virtual void SetAssetAPI(IAssetGenerationAPI* AssetAPI);
+
+	virtual bool HasCancel() const override { return true; }
+	virtual bool HasAccept() const override;
+	virtual bool CanAccept() const override;
+
+protected:
+
+	UPROPERTY()
+	UGenerateStaticMeshLODAssetToolProperties* BasicProperties;
+
+protected:
+	UWorld* TargetWorld;
+	IAssetGenerationAPI* AssetAPI;
+
+	TPimplPtr<FGenerateStaticMeshLODProcess> GenerateProcess;
+
+	void CreateNewAsset();
+	void UpdateExistingAsset();
+};
