MeshLODToolset:
- add support for extracting tangents in GenerateMeshLODGraph
- add GenerateStaticMeshLODProcess::UpdateSourceAsset(), writes generated mesh to LOD1, adds material slots for generated materials
- add GeneratedStaticMeshLODProcess::GetDerivedMaterialsPreview(), creates MIDs and temporary UTexture2Ds for generated textures, returns as a set that can be passed to a preview mesh
- add option to UGenerateStaticMeshLODAssetTool to update source asset instead of creating new static mesh
- UGenerateStaticMeshLODAssetTool now initializes result at startup, shows side-by-side preview

#rnx

diff --git a/Engine/Plugins/Experimental/GeometryFlow/Source/GeometryFlowMeshProcessing/Public/MeshProcessingNodes/MeshTangentsNodes.h b/Engine/Plugins/Experimental/GeometryFlow/Source/GeometryFlowMeshProcessing/Public/MeshProcessingNodes/MeshTangentsNodes.h
--- a/Engine/Plugins/Experimental/GeometryFlow/Source/GeometryFlowMeshProcessing/Public/MeshProcessingNodes/MeshTangentsNodes.h
+++ b/Engine/Plugins/Experimental/GeometryFlow/Source/GeometryFlowMeshProcessing/Public/MeshProcessingNodes/MeshTangentsNodes.h
@@ -33,6 +33,7 @@ struct FMeshTangentsSettings
 GEOMETRYFLOW_DECLARE_SETTINGS_TYPES(FMeshTangentsSettings, Tangents);
 
 
+typedef TTransferNode<FMeshTangentsd, (int)EMeshProcessingDataTypes::MeshTangentSet> FMeshTangentsTransferNode;
 
 
 class FComputeMeshTangentsNode : public TTransformerWithSettingsNode<
diff --git a/Engine/Plugins/Experimental/MeshLODToolset/Source/MeshLODToolset/Private/Graphs/GenerateMeshLODGraph.cpp b/Engine/Plugins/Experimental/MeshLODToolset/Source/MeshLODToolset/Private/Graphs/GenerateMeshLODGraph.cpp
--- a/Engine/Plugins/Experimental/MeshLODToolset/Source/MeshLODToolset/Private/Graphs/GenerateMeshLODGraph.cpp
+++ b/Engine/Plugins/Experimental/MeshLODToolset/Source/MeshLODToolset/Private/Graphs/GenerateMeshLODGraph.cpp
@@ -39,6 +39,7 @@ void FGenerateMeshLODGraph::SetSourceMesh(const FDynamicMesh3& SourceMeshIn)
 
 void FGenerateMeshLODGraph::EvaluateResult(
 	FDynamicMesh3& ResultMesh,
+	FMeshTangentsd& ResultTangents,
 	FSimpleShapeSet3d& ResultCollision,
 	UE::GeometryFlow::FNormalMapImage& NormalMap,
 	TArray<TUniquePtr<UE::GeometryFlow::FTextureImage>>& TextureImages)
@@ -85,6 +86,20 @@ void FGenerateMeshLODGraph::EvaluateResult(
 
 	UE_LOG(LogTemp, Warning, TEXT("OutputCollisionPass - Evaluated %d Nodes, Recomputed %d"), CollisionEvalInfo->NumEvaluations(), CollisionEvalInfo->NumComputes());
 
+	// 
+	// evaluate tangents
+	//
+
+	bool bTakeResultTangents = false;
+	ResultTangents = FMeshTangentsd();
+
+	TUniquePtr<FEvaluationInfo> TangentsEvalInfo = MakeUnique<FEvaluationInfo>();
+	EGeometryFlowResult TangentsEvalResult = Graph->EvaluateResult(TangentsOutputNode, FMeshTangentsTransferNode::OutParamValue(),
+		ResultTangents, (int)EMeshProcessingDataTypes::MeshTangentSet, TangentsEvalInfo, bTakeResultTangents);
+	ensure(TangentsEvalResult == EGeometryFlowResult::Ok);
+
+	UE_LOG(LogTemp, Warning, TEXT("OutputTangentsPass - Evaluated %d Nodes, Recomputed %d"), TangentsEvalInfo->NumEvaluations(), TangentsEvalInfo->NumComputes());
+
 
 	//
 	// evaluate result mesh
@@ -165,6 +180,10 @@ void FGenerateMeshLODGraph::BuildGraph()
 	FGraph::FHandle TangentsSettingsNode = Graph->AddNodeOfType<FTangentsSettingsSourceNode>(TEXT("TangentsSettings"));
 	Graph->InferConnection(TangentsSettingsNode, TangentsNode);
 
+	// tangents output
+	TangentsOutputNode = Graph->AddNodeOfType<FMeshTangentsTransferNode>(TEXT("OutputTangents"));
+	Graph->InferConnection(TangentsNode, TangentsOutputNode);
+
 	// create bake cache
 
 	BakeCacheNode = Graph->AddNodeOfType<FMakeMeshBakingCacheNode>(TEXT("MakeBakeCache"));
@@ -199,7 +218,7 @@ void FGenerateMeshLODGraph::BuildGraph()
 	FGraph::FHandle GenerateConvexesSettingsNode = Graph->AddNodeOfType<FGenerateConvexHullsCollisionSettingsSourceNode>(TEXT("GenerateConvexesSettings"));
 	Graph->InferConnection(GenerateConvexesSettingsNode, GenerateConvexesNode);
 
-	// final mesh output
+	// final collision output
 
 	CollisionOutputNode = Graph->AddNodeOfType<FCollisionGeometryTransferNode>(TEXT("OutputCollision"));
 	Graph->InferConnection(GenerateConvexesNode, CollisionOutputNode);
diff --git a/Engine/Plugins/Experimental/MeshLODToolset/Source/MeshLODToolset/Private/Graphs/GenerateStaticMeshLODProcess.cpp b/Engine/Plugins/Experimental/MeshLODToolset/Source/MeshLODToolset/Private/Graphs/GenerateStaticMeshLODProcess.cpp
--- a/Engine/Plugins/Experimental/MeshLODToolset/Source/MeshLODToolset/Private/Graphs/GenerateStaticMeshLODProcess.cpp
+++ b/Engine/Plugins/Experimental/MeshLODToolset/Source/MeshLODToolset/Private/Graphs/GenerateStaticMeshLODProcess.cpp
@@ -7,6 +7,7 @@
 
 #include "MeshDescriptionToDynamicMesh.h"
 #include "DynamicMeshToMeshDescription.h"
+#include "DynamicMeshAttributeSet.h"
 
 #include "AssetUtils/Texture2DUtil.h"
 #include "AssetUtils/Texture2DBuilder.h"
@@ -20,15 +21,19 @@
 #include "IAssetTools.h"
 #include "FileHelpers.h"
 
+#include "Engine/Engine.h"
+#include "Editor.h"
+
 #include "Materials/Material.h"
 #include "Materials/MaterialInstanceConstant.h"
 #include "Factories/MaterialInstanceConstantFactoryNew.h"
+#include "Materials/MaterialInstanceDynamic.h"
 
 #include "Engine/Classes/Engine/StaticMesh.h"
 #include "Engine/Classes/Components/StaticMeshComponent.h"
 #include "Engine/Classes/PhysicsEngine/BodySetup.h"
 
-
+#define LOCTEXT_NAMESPACE "FGenerateStaticMeshLODProcess"
 
 
 namespace
@@ -146,6 +151,8 @@ bool FGenerateStaticMeshLODProcess::Initialize(UStaticMesh* StaticMeshIn)
 
 	CalculateDerivedPathName(GetDefaultDerivedAssetSuffix());
 
+	InitializeGenerator();
+
 	return true;
 }
 
@@ -164,8 +171,7 @@ void FGenerateStaticMeshLODProcess::CalculateDerivedPathName(FString NewAssetSuf
 }
 
 
-
-bool FGenerateStaticMeshLODProcess::ComputeDerivedSourceData()
+bool FGenerateStaticMeshLODProcess::InitializeGenerator()
 {
 	Generator = MakeUnique<FGenerateMeshLODGraph>();
 	Generator->BuildGraph();
@@ -176,7 +182,7 @@ bool FGenerateStaticMeshLODProcess::ComputeDerivedSourceData()
 		for (const FTextureInfo& TexInfo : MatInfo.SourceTextures)
 		{
 			if (TexInfo.bShouldBakeTexture &&
-				(TextureToDerivedTexIndex.Contains(TexInfo.SourceTexture) == false) )
+				(TextureToDerivedTexIndex.Contains(TexInfo.SourceTexture) == false))
 			{
 				int32 NewIndex = Generator->AppendTextureBakeNode(TexInfo.Image, TexInfo.SourceTexture->GetName());
 				TextureToDerivedTexIndex.Add(TexInfo.SourceTexture, NewIndex);
@@ -186,8 +192,18 @@ bool FGenerateStaticMeshLODProcess::ComputeDerivedSourceData()
 
 	Generator->SetSourceMesh(this->SourceMesh);
 
-	Generator->EvaluateResult(this->DerivedLODMesh, this->DerivedCollision, this->DerivedNormalMapImage, this->DerivedTextureImages);
+	return true;
+}
+
 
+bool FGenerateStaticMeshLODProcess::ComputeDerivedSourceData()
+{
+	Generator->EvaluateResult(
+		this->DerivedLODMesh, 
+		this->DerivedLODMeshTangents,
+		this->DerivedCollision, 
+		this->DerivedNormalMapImage, 
+		this->DerivedTextureImages);
 
 	// copy all materials for now...we are going to replace all the images though, and
 	// should not copy those?
@@ -213,6 +229,108 @@ bool FGenerateStaticMeshLODProcess::ComputeDerivedSourceData()
 
 
 
+
+void FGenerateStaticMeshLODProcess::GetDerivedMaterialsPreview(FPreviewMaterials& MaterialSetOut)
+{
+	// force garbage collection of outstanding preview materials
+	CollectGarbage(GARBAGE_COLLECTION_KEEPFLAGS);
+
+
+	// create derived textures
+	int32 NumMaterials = SourceMaterials.Num();
+	check(DerivedMaterials.Num() == NumMaterials);
+	TMap<UTexture2D*, UTexture2D*> SourceToPreviewTexMap;
+	for (int32 mi = 0; mi < NumMaterials; ++mi)
+	{
+		const FMaterialInfo& SourceMaterialInfo = SourceMaterials[mi];
+		const FMaterialInfo& DerivedMaterialInfo = DerivedMaterials[mi];
+
+		int32 NumTextures = SourceMaterialInfo.SourceTextures.Num();
+		check(DerivedMaterialInfo.SourceTextures.Num() == NumTextures);
+		for (int32 ti = 0; ti < NumTextures; ++ti)
+		{
+			const FTextureInfo& SourceTex = SourceMaterialInfo.SourceTextures[ti];
+			bool bConvertToSRGB = SourceTex.SourceTexture->SRGB;
+			const FTextureInfo& DerivedTex = DerivedMaterialInfo.SourceTextures[ti];
+			if (DerivedTex.bShouldBakeTexture)
+			{
+				FTexture2DBuilder TextureBuilder;
+				TextureBuilder.Initialize(FTexture2DBuilder::ETextureType::Color, DerivedTex.Dimensions);
+				TextureBuilder.GetTexture2D()->SRGB = bConvertToSRGB;
+				TextureBuilder.Copy(DerivedTex.Image, bConvertToSRGB);
+				TextureBuilder.Commit(false);
+				UTexture2D* PreviewTex = TextureBuilder.GetTexture2D();
+				if (ensure(PreviewTex))
+				{
+					SourceToPreviewTexMap.Add(SourceTex.SourceTexture, PreviewTex);
+					MaterialSetOut.Textures.Add(PreviewTex);
+				}
+			}
+		}
+	}
+
+	// create derived normal map texture
+	FTexture2DBuilder NormapMapBuilder;
+	NormapMapBuilder.Initialize(FTexture2DBuilder::ETextureType::NormalMap, DerivedNormalMapImage.Image.GetDimensions());
+	NormapMapBuilder.Copy(DerivedNormalMapImage.Image, false);
+	NormapMapBuilder.Commit(false);
+	UTexture2D* PreviewNormalMapTex = NormapMapBuilder.GetTexture2D();
+	MaterialSetOut.Textures.Add(PreviewNormalMapTex);
+
+	// create derived MIDs and point to new textures
+	for (int32 mi = 0; mi < NumMaterials; ++mi)
+	{
+		const FMaterialInfo& SourceMaterialInfo = SourceMaterials[mi];
+
+		UMaterialInterface* MaterialInterface = SourceMaterialInfo.SourceMaterial.MaterialInterface;
+		UMaterialInstanceDynamic* GeneratedMID = UMaterialInstanceDynamic::Create(MaterialInterface, NULL);
+
+		// rewrite texture parameters to new textures
+		UpdateMaterialTextureParameters(GeneratedMID, SourceMaterialInfo, SourceToPreviewTexMap, PreviewNormalMapTex);
+
+		MaterialSetOut.Materials.Add(GeneratedMID);
+	}
+
+}
+
+
+void FGenerateStaticMeshLODProcess::UpdateMaterialTextureParameters(
+	UMaterialInstanceDynamic* Material, 
+	const FMaterialInfo& SourceMaterialInfo,
+	const TMap<UTexture2D*, UTexture2D*>& PreviewTextures, 
+	UTexture2D* PreviewNormalMap)
+{
+	Material->Modify();
+	int32 NumTextures = SourceMaterialInfo.SourceTextures.Num();
+	for (int32 ti = 0; ti < NumTextures; ++ti)
+	{
+		const FTextureInfo& SourceTex = SourceMaterialInfo.SourceTextures[ti];
+
+		if (SourceTex.bIsNormalMap)
+		{
+			if (ensure(PreviewNormalMap))
+			{
+				FMaterialParameterInfo ParamInfo(SourceTex.ParameterName);
+				Material->SetTextureParameterValueByInfo(ParamInfo, PreviewNormalMap);
+			}
+		}
+		else if (SourceTex.bShouldBakeTexture)
+		{
+			UTexture2D*const* FoundTexture = PreviewTextures.Find(SourceTex.SourceTexture);
+			if (ensure(FoundTexture))
+			{
+				FMaterialParameterInfo ParamInfo(SourceTex.ParameterName);
+				Material->SetTextureParameterValueByInfo(ParamInfo, *FoundTexture);
+			}
+		}
+	}
+	Material->PostEditChange();
+}
+
+
+
+
+
 bool FGenerateStaticMeshLODProcess::WriteDerivedAssetData()
 {
 	WriteDerivedTextures();
@@ -225,6 +343,15 @@ bool FGenerateStaticMeshLODProcess::WriteDerivedAssetData()
 }
 
 
+void FGenerateStaticMeshLODProcess::UpdateSourceAsset()
+{
+	WriteDerivedTextures();
+
+	WriteDerivedMaterials();
+
+	UpdateSourceStaticMeshAsset();
+}
+
 
 void FGenerateStaticMeshLODProcess::WriteDerivedTextures()
 {
@@ -503,4 +630,89 @@ void FGenerateStaticMeshLODProcess::WriteDerivedStaticMeshAsset()
 
 	// done updating mesh
 	GeneratedStaticMesh->PostEditChange();
-}
\ No newline at end of file
+}
+
+
+
+
+void FGenerateStaticMeshLODProcess::UpdateSourceStaticMeshAsset()
+{
+	GEditor->BeginTransaction(LOCTEXT("UpdateExistingAssetMessage", "Added Generated LOD"));
+
+	SourceStaticMesh->Modify();
+
+	TArray<FStaticMaterial> ExistingMaterials = SourceStaticMesh->GetStaticMaterials();
+
+	// append new materials to material set
+	TArray<int32> NewMatIndexMap;
+	for (FMaterialInfo& DerivedMaterialInfo : DerivedMaterials)
+	{
+		int32 NewIndex = ExistingMaterials.Num();
+		ExistingMaterials.Add(DerivedMaterialInfo.SourceMaterial);
+		NewMatIndexMap.Add(NewIndex);
+	}
+
+	// rewrite material IDs on derived mesh
+	DerivedLODMesh.Attributes()->EnableMaterialID();
+	FDynamicMeshMaterialAttribute* MaterialIDs = DerivedLODMesh.Attributes()->GetMaterialID();
+	for (int32 tid : DerivedLODMesh.TriangleIndicesItr())
+	{
+		int32 CurMaterialID = MaterialIDs->GetValue(tid);
+		int32 NewMaterialID = NewMatIndexMap[CurMaterialID];
+		MaterialIDs->SetValue(tid, NewMaterialID);
+	}
+
+	// update materials on generated mesh
+	SourceStaticMesh->SetStaticMaterials(ExistingMaterials);
+
+	// store new derived LOD as LOD 1
+	SourceStaticMesh->SetNumSourceModels(2);
+	FMeshDescription* MeshDescription = SourceStaticMesh->GetMeshDescription(1);
+	if (MeshDescription == nullptr)
+	{
+		MeshDescription = SourceStaticMesh->CreateMeshDescription(1);
+	}
+	FConversionToMeshDescriptionOptions ConversionOptions;
+	FDynamicMeshToMeshDescription Converter(ConversionOptions);
+	Converter.Convert(&DerivedLODMesh, *MeshDescription);
+	SourceStaticMesh->CommitMeshDescription(1);
+
+	// this will prevent simplification?
+	FStaticMeshSourceModel& SrcModel = SourceStaticMesh->GetSourceModel(1);
+	SrcModel.ReductionSettings.MaxDeviation = 0.0f;
+	SrcModel.ReductionSettings.PercentTriangles = 1.0f;
+	SrcModel.ReductionSettings.PercentVertices = 1.0f;
+
+	// collision
+	FPhysicsDataCollection NewCollisionGeo;
+	NewCollisionGeo.Geometry = DerivedCollision;
+	NewCollisionGeo.CopyGeometryToAggregate();
+
+	// code below derived from FStaticMeshEditor::DuplicateSelectedPrims()
+	UBodySetup* BodySetup = SourceStaticMesh->GetBodySetup();
+	// mark the BodySetup for modification. Do we need to modify the UStaticMesh??
+	BodySetup->Modify();
+	//Clear the cache (PIE may have created some data), create new GUID    (comment from StaticMeshEditor)
+	BodySetup->InvalidatePhysicsData();
+	BodySetup->RemoveSimpleCollision();
+	BodySetup->AggGeom = NewCollisionGeo.AggGeom;
+	// update collision type
+	BodySetup->CollisionTraceFlag = ECollisionTraceFlag::CTF_UseDefault;
+	// rebuild physics data
+	BodySetup->InvalidatePhysicsData();
+	BodySetup->CreatePhysicsMeshes();
+
+	// do we need to do a post edit change here??
+
+	// is this necessary? 
+	SourceStaticMesh->CreateNavCollision(/*bIsUpdate=*/true);
+
+	GEditor->EndTransaction();
+
+	// done updating mesh
+	SourceStaticMesh->PostEditChange();
+}
+
+
+
+#undef LOCTEXT_NAMESPACE
\ No newline at end of file
diff --git a/Engine/Plugins/Experimental/MeshLODToolset/Source/MeshLODToolset/Private/Tools/GenerateStaticMeshLODAssetTool.cpp b/Engine/Plugins/Experimental/MeshLODToolset/Source/MeshLODToolset/Private/Tools/GenerateStaticMeshLODAssetTool.cpp
--- a/Engine/Plugins/Experimental/MeshLODToolset/Source/MeshLODToolset/Private/Tools/GenerateStaticMeshLODAssetTool.cpp
+++ b/Engine/Plugins/Experimental/MeshLODToolset/Source/MeshLODToolset/Private/Tools/GenerateStaticMeshLODAssetTool.cpp
@@ -13,6 +13,7 @@
 #include "Selection/SelectClickedAction.h"
 #include "DynamicMeshEditor.h"
 #include "MeshTransforms.h"
+#include "MeshTangents.h"
 
 #include "MeshDescriptionToDynamicMesh.h"
 #include "DynamicMeshToMeshDescription.h"
@@ -109,8 +110,6 @@ void UGenerateStaticMeshLODAssetTool::Setup()
 		LOCTEXT("OnStartStaticMeshLODAssetTool", "This tool creates a new LOD asset"),
 		EToolMessageLevel::UserNotification);
 
-
-
 	GenerateProcess = MakePimpl<FGenerateStaticMeshLODProcess>();
 
 	TUniquePtr<FPrimitiveComponentTarget>& SourceComponent = ComponentTargets[0];
@@ -121,6 +120,17 @@ void UGenerateStaticMeshLODAssetTool::Setup()
 		GenerateProcess->Initialize(StaticMesh);
 	}
 
+	FBoxSphereBounds Bounds = StaticMeshComponent->Bounds;
+	FTransform PreviewTransform = SourceComponent->GetWorldTransform();
+	PreviewTransform.AddToTranslation(FVector(0, 2.0f*Bounds.SphereRadius, 0));
+
+	PreviewMesh = NewObject<UPreviewMesh>(this);
+	PreviewMesh->CreateInWorld(TargetWorld, PreviewTransform);
+	PreviewMesh->SetVisible(true);
+	PreviewMesh->SetTangentsMode(EDynamicMeshTangentCalcType::ExternallyCalculated);
+
+	bPreviewValid = false;
+	ValidatePreview();
 }
 
 
@@ -129,10 +139,21 @@ void UGenerateStaticMeshLODAssetTool::Shutdown(EToolShutdownType ShutdownType)
 {
 	BasicProperties->SaveProperties(this);
 
+	PreviewMesh->SetVisible(false);
+	PreviewMesh->Disconnect();
+	PreviewMesh = nullptr;
+
 	if (ShutdownType == EToolShutdownType::Accept)
 	{
-		CreateNewAsset();
-		//UpdateExistingAsset();
+		if (BasicProperties->OutputMode == EGenerateLODAssetOutputMode::UpdateExistingAsset)
+		{
+			UpdateExistingAsset();
+		}
+		else
+		{
+			CreateNewAsset();
+		}
+
 	}
 }
 
@@ -151,249 +172,49 @@ bool UGenerateStaticMeshLODAssetTool::CanAccept() const
 	return true;
 }
 
-
-void UGenerateStaticMeshLODAssetTool::CreateNewAsset()
+void UGenerateStaticMeshLODAssetTool::ValidatePreview()
 {
-	IAssetTools& AssetTools = FModuleManager::Get().LoadModuleChecked<FAssetToolsModule>("AssetTools").Get();
-
-	//GetToolManager()->BeginUndoTransaction( LOCTEXT("GenerateStaticMeshLODAssetTransaction", "Generate LOD Asset") );
+	if (bPreviewValid) return;
 
-	GenerateProcess->CalculateDerivedPathName(BasicProperties->GeneratedSuffix);
 	GenerateProcess->ComputeDerivedSourceData();
-	GenerateProcess->WriteDerivedAssetData();
-
-/*
-	UStaticMesh* StaticMesh = GenerateProcess->GetSourceStaticMesh();
-	const FDynamicMesh3& GeneratedLOD0Mesh = GenerateProcess->GetDerivedLOD0Mesh();
-
-	// find existing path and generate AutoLOD path
-	FString FullSMPathWithExtension = UEditorAssetLibrary::GetPathNameForLoadedAsset(StaticMesh);
-	FString SMPathWithName = FPaths::GetBaseFilename(FullSMPathWithExtension, false);
-	FString SMPath = FPaths::GetPath(SMPathWithName);
-	FString UseSuffix = FPaths::MakeValidFileName(BasicProperties->GeneratedSuffix);
-	if (UseSuffix.Len() == 0)
-	{
-		UseSuffix = TEXT("_AutoLOD");
-	}
-	FString GeneratedSMPath = FString::Printf(TEXT("%s%s"), *SMPathWithName, *UseSuffix);
-
-	// [TODO] should we try to re-use existing asset here, or should we delete it? 
-	// The source asset might have had any number of config changes that we want to
-	// preserve in the duplicate...
-	UStaticMesh* GeneratedStaticMesh = nullptr;
-	if (UEditorAssetLibrary::DoesAssetExist(GeneratedSMPath))
+	const FDynamicMesh3& ResultMesh = GenerateProcess->GetDerivedLOD0Mesh();
+	const FMeshTangentsd& ResultTangents = GenerateProcess->GetDerivedLOD0MeshTangents();
+	
+	PreviewMesh->EditMesh([&](FDynamicMesh3& MeshToUpdate)
 	{
-		UObject* LoadedAsset = UEditorAssetLibrary::LoadAsset(GeneratedSMPath);
-		GeneratedStaticMesh = Cast<UStaticMesh>(LoadedAsset);
-	}
-	else
-	{
-		UObject* DupeAsset = UEditorAssetLibrary::DuplicateAsset(SMPathWithName, GeneratedSMPath);
-		GeneratedStaticMesh = Cast<UStaticMesh>(DupeAsset);
-	}
-
-	// make sure transactional flag is on
-	GeneratedStaticMesh->SetFlags(RF_Transactional);
-	GeneratedStaticMesh->Modify();
-
-	// update MeshDescription LOD0 mesh
-	GeneratedStaticMesh->SetNumSourceModels(1);
-	FMeshDescription* MeshDescription = GeneratedStaticMesh->GetMeshDescription(0);
-	FConversionToMeshDescriptionOptions ConversionOptions;
-	FDynamicMeshToMeshDescription Converter(ConversionOptions);
-	Converter.Convert(&GeneratedLOD0Mesh, *MeshDescription);
-	GeneratedStaticMesh->CommitMeshDescription(0);
-
-	// done updating mesh
-	GeneratedStaticMesh->PostEditChange();
+		MeshToUpdate = GenerateProcess->GetDerivedLOD0Mesh();
+	});
+	PreviewMesh->UpdateTangents(&ResultTangents, true);
 
-
-
-	const TArray<FStaticMaterial>& OldMaterials = StaticMesh->GetStaticMaterials();
-	TArray<FStaticMaterial>& NewMaterials = GeneratedStaticMesh->GetStaticMaterials();
-	check(OldMaterials.Num() == NewMaterials.Num());
-
-	for (int32 mi = 0; mi < OldMaterials.Num(); ++mi)
+	FGenerateStaticMeshLODProcess::FPreviewMaterials PreviewMaterialSet;
+	GenerateProcess->GetDerivedMaterialsPreview(PreviewMaterialSet);
+	if (PreviewMaterialSet.Materials.Num() > 0)
 	{
-		UMaterialInterface* MaterialInterface = OldMaterials[mi].MaterialInterface;
-		bool bSourceIsMIC = (Cast<UMaterialInstanceConstant>(MaterialInterface) != nullptr);
-
-		FString SourceMaterialPath = UEditorAssetLibrary::GetPathNameForLoadedAsset(MaterialInterface);
-		FString MaterialName = FPaths::GetBaseFilename(SourceMaterialPath, true);
-		FString NewMaterialName = FString::Printf(TEXT("%s%s"), *MaterialName, *UseSuffix);
-		FString NewMaterialPath = FPaths::Combine(SMPath, NewMaterialName);
-
+		PreviewTextures = PreviewMaterialSet.Textures;
+		PreviewMaterials = PreviewMaterialSet.Materials;
+		PreviewMesh->SetMaterials(PreviewMaterials);
 
-		// delete existing material so that we can have a clean duplicate
-		bool bNewMaterialExists = UEditorAssetLibrary::DoesAssetExist(NewMaterialPath);
-		if (bNewMaterialExists)
-		{
-			bool bDeleteOK = UEditorAssetLibrary::DeleteAsset(NewMaterialPath);
-			ensure(bDeleteOK);
-		}
-
-		// If source is a MIC, we can just duplicate it. If it is a UMaterial, we want to
-		// create a child MIC? Or we could dupe the Material and rewrite the textures.
-		// Probably needs to be an option.
-		UMaterialInstanceConstant* GeneratedMIC = nullptr;
-		if (bSourceIsMIC)
-		{
-			UObject* DupeAsset = UEditorAssetLibrary::DuplicateAsset(SourceMaterialPath, NewMaterialPath);
-			GeneratedMIC = Cast<UMaterialInstanceConstant>(DupeAsset);
-		}
-		else
-		{
-			UMaterial* SourceMaterial = MaterialInterface->GetBaseMaterial();
-			if (ensure(SourceMaterial))
-			{
-				UMaterialInstanceConstantFactoryNew* Factory = NewObject<UMaterialInstanceConstantFactoryNew>();
-				Factory->InitialParent = SourceMaterial;
-
-				UObject* NewAsset = AssetTools.CreateAsset(NewMaterialName, FPackageName::GetLongPackagePath(NewMaterialPath),
-					UMaterialInstanceConstant::StaticClass(), Factory);
-
-				GeneratedMIC = Cast<UMaterialInstanceConstant>(NewAsset);
-			}
-		}
-
-		NewMaterials[mi].MaterialInterface = GeneratedMIC;
+		BasicProperties->PreviewTextures = PreviewTextures;
 	}
 
-
-	// update materials on generated mesh
-	GeneratedStaticMesh->SetStaticMaterials(NewMaterials);
-
-*/
-
-	// I think this would let us jump to new assets...
-	//ObjectsToSync.Add(NewAsset);
-	//GEditor->SyncBrowserToObjects(ObjectsToSync);
-
-
-	//GetToolManager()->EndUndoTransaction();
+	bPreviewValid = true;
 }
 
 
-
-
+void UGenerateStaticMeshLODAssetTool::CreateNewAsset()
+{
+	GenerateProcess->CalculateDerivedPathName(BasicProperties->GeneratedSuffix);
+	//GenerateProcess->ComputeDerivedSourceData();
+	GenerateProcess->WriteDerivedAssetData();
+}
 
 
 
 void UGenerateStaticMeshLODAssetTool::UpdateExistingAsset()
 {
-	GetToolManager()->BeginUndoTransaction(LOCTEXT("GenerateStaticMeshLODAssetToolTransactionName", "Combine Meshes"));
-
-	// note there is a MergeMeshUtilities.h w/ a very feature-filled mesh merging class, but for simplicity (and to fit modeling tool needs)
-	// this tool currently converts things through dynamic mesh instead
-
-	AActor* SkipActor = nullptr;
-
-#if WITH_EDITOR
-
-	bool bMergeSameMaterials = true;
-	TArray<UMaterialInterface*> AllMaterials;
-	TMap<UMaterialInterface*, int> KnownMaterials;
-	TArray<int> CombinedMatToOutMatIdx;
-	for (int32 ComponentIdx = 0; ComponentIdx < ComponentTargets.Num(); ComponentIdx++)
-	{
-		TUniquePtr<FPrimitiveComponentTarget>& ComponentTarget = ComponentTargets[ComponentIdx];
-		for (int MaterialIdx = 0, NumMaterials = ComponentTarget->GetNumMaterials(); MaterialIdx < NumMaterials; MaterialIdx++)
-		{
-			UMaterialInterface* Mat = ComponentTarget->GetMaterial(MaterialIdx);
-			int32 OutMatIdx = -1;
-			if (!bMergeSameMaterials || !KnownMaterials.Contains(Mat))
-			{
-				OutMatIdx = AllMaterials.Num();
-				if (bMergeSameMaterials)
-				{
-					KnownMaterials.Add(Mat, AllMaterials.Num());
-				}
-				AllMaterials.Add(Mat);
-			}
-			else
-			{
-				OutMatIdx = KnownMaterials[Mat];
-			}
-			CombinedMatToOutMatIdx.Add(OutMatIdx);
-		}
-	}
-
-	FDynamicMesh3 AccumulateDMesh;
-	AccumulateDMesh.EnableTriangleGroups();
-	AccumulateDMesh.EnableAttributes();
-	AccumulateDMesh.Attributes()->EnableMaterialID();
-
-	int32 SkipIndex = 0;
-	TUniquePtr<FPrimitiveComponentTarget>& UpdateTarget = ComponentTargets[SkipIndex];
-	SkipActor = UpdateTarget->GetOwnerActor();
-
-	FTransform3d TargetToWorld = (FTransform3d)UpdateTarget->GetWorldTransform();
-	FTransform3d WorldToTarget = TargetToWorld.Inverse();
-
-	{
-		FScopedSlowTask SlowTask(ComponentTargets.Num() + 1,
-			LOCTEXT("GenerateStaticMeshLODAssetBuild", "Building combined mesh ..."));
-		SlowTask.MakeDialog();
-
-		int MatIndexBase = 0;
-		for (int32 ComponentIdx = 0; ComponentIdx < ComponentTargets.Num(); ComponentIdx++)
-		{
-			SlowTask.EnterProgressFrame(1);
-			TUniquePtr<FPrimitiveComponentTarget>& ComponentTarget = ComponentTargets[ComponentIdx];
-
-			FMeshDescriptionToDynamicMesh Converter;
-			FDynamicMesh3 ComponentDMesh;
-			Converter.Convert(ComponentTarget->GetMesh(), ComponentDMesh);
-
-			// update material IDs to account for combined material set
-			FDynamicMeshMaterialAttribute* MatAttrib = ComponentDMesh.Attributes()->GetMaterialID();
-			for (int TID : ComponentDMesh.TriangleIndicesItr())
-			{
-				MatAttrib->SetValue(TID, CombinedMatToOutMatIdx[MatIndexBase + MatAttrib->GetValue(TID)]);
-			}
-			MatIndexBase += ComponentTarget->GetNumMaterials();
-
-
-			if (ComponentIdx != SkipIndex)
-			{
-				FTransform3d ComponentToWorld = (FTransform3d)ComponentTarget->GetWorldTransform();
-				MeshTransforms::ApplyTransform(ComponentDMesh, ComponentToWorld);
-				if (ComponentToWorld.GetDeterminant() < 0)
-				{
-					ComponentDMesh.ReverseOrientation(true);
-				}
-				MeshTransforms::ApplyTransform(ComponentDMesh, WorldToTarget);
-				if (WorldToTarget.GetDeterminant() < 0)
-				{
-					ComponentDMesh.ReverseOrientation(true);
-				}
-			}
-
-			FDynamicMeshEditor Editor(&AccumulateDMesh);
-			FMeshIndexMappings IndexMapping;
-			Editor.AppendMesh(&ComponentDMesh, IndexMapping);
-		}
-
-		SlowTask.EnterProgressFrame(1);
-
-		UpdateTarget->CommitMesh([&](const FPrimitiveComponentTarget::FCommitParams& CommitParams)
-		{
-			FDynamicMeshToMeshDescription Converter;
-			Converter.Convert(&AccumulateDMesh, *CommitParams.MeshDescription);
-		});
-
-		FComponentMaterialSet MaterialSet;
-		MaterialSet.Materials = AllMaterials;
-		UpdateTarget->CommitMaterialSetUpdate(MaterialSet, true);
-
-		// select the new actor
-		ToolSelectionUtil::SetNewActorSelection(GetToolManager(), SkipActor);
-	}
-#endif
-
-
-
-	GetToolManager()->EndUndoTransaction();
+	GenerateProcess->CalculateDerivedPathName(BasicProperties->GeneratedSuffix);
+	//GenerateProcess->ComputeDerivedSourceData();
+	GenerateProcess->UpdateSourceAsset();
 }
 
 
diff --git a/Engine/Plugins/Experimental/MeshLODToolset/Source/MeshLODToolset/Public/Graphs/GenerateMeshLODGraph.h b/Engine/Plugins/Experimental/MeshLODToolset/Source/MeshLODToolset/Public/Graphs/GenerateMeshLODGraph.h
--- a/Engine/Plugins/Experimental/MeshLODToolset/Source/MeshLODToolset/Public/Graphs/GenerateMeshLODGraph.h
+++ b/Engine/Plugins/Experimental/MeshLODToolset/Source/MeshLODToolset/Public/Graphs/GenerateMeshLODGraph.h
@@ -7,6 +7,7 @@
 #include "DataTypes/TextureImageData.h"
 #include "DataTypes/CollisionGeometryData.h"
 #include "DynamicMesh3.h"
+#include "MeshTangents.h"
 
 
 class FGenerateMeshLODGraph
@@ -21,6 +22,7 @@ public:
 
 	void EvaluateResult(
 		FDynamicMesh3& ResultMesh,
+		FMeshTangentsd& ResultTangents,
 		FSimpleShapeSet3d& ResultCollision,
 		UE::GeometryFlow::FNormalMapImage& NormalMap,
 		TArray<TUniquePtr<UE::GeometryFlow::FTextureImage>>& TextureImages);
@@ -62,5 +64,6 @@ protected:
 
 	UE::GeometryFlow::FGraph::FHandle CollisionOutputNode;
 	UE::GeometryFlow::FGraph::FHandle MeshOutputNode;
+	UE::GeometryFlow::FGraph::FHandle TangentsOutputNode;
 
 };
\ No newline at end of file
diff --git a/Engine/Plugins/Experimental/MeshLODToolset/Source/MeshLODToolset/Public/Graphs/GenerateStaticMeshLODProcess.h b/Engine/Plugins/Experimental/MeshLODToolset/Source/MeshLODToolset/Public/Graphs/GenerateStaticMeshLODProcess.h
--- a/Engine/Plugins/Experimental/MeshLODToolset/Source/MeshLODToolset/Public/Graphs/GenerateStaticMeshLODProcess.h
+++ b/Engine/Plugins/Experimental/MeshLODToolset/Source/MeshLODToolset/Public/Graphs/GenerateStaticMeshLODProcess.h
@@ -11,6 +11,7 @@
 
 class UTexture2D;
 class UMaterialInstanceConstant;
+class UMaterialInstanceDynamic;
 
 class FGenerateStaticMeshLODProcess
 {
@@ -29,14 +30,31 @@ public:
 
 	void CalculateDerivedPathName(FString NewAssetSuffix);
 
-
 	bool ComputeDerivedSourceData();
 	const FDynamicMesh3& GetDerivedLOD0Mesh() const { return DerivedLODMesh; }
+	const FMeshTangentsd& GetDerivedLOD0MeshTangents() const { return DerivedLODMeshTangents; }
 
 
+	/**
+	 * Creates new Asset
+	 */
 	bool WriteDerivedAssetData();
 
 
+	/**
+	 * Updates existing SM Asset
+	 */
+	void UpdateSourceAsset();
+
+
+
+	struct FPreviewMaterials
+	{
+		TArray<UMaterialInterface*> Materials;
+		TArray<UTexture2D*> Textures;
+	};
+	void GetDerivedMaterialsPreview(FPreviewMaterials& MaterialSetOut);
+
 
 protected:
 
@@ -73,6 +91,7 @@ protected:
 	FString DerivedAssetPath;
 
 	FDynamicMesh3 DerivedLODMesh;
+	FMeshTangentsd DerivedLODMeshTangents;
 	FSimpleShapeSet3d DerivedCollision;
 	UE::GeometryFlow::FNormalMapImage DerivedNormalMapImage;
 	TArray<TUniquePtr<UE::GeometryFlow::FTextureImage>> DerivedTextureImages;
@@ -82,6 +101,7 @@ protected:
 	UTexture2D* DerivedNormalMapTex;
 
 	TUniquePtr<FGenerateMeshLODGraph> Generator;
+	bool InitializeGenerator();
 
 	bool WriteDerivedTexture(UTexture2D* SourceTexture, UTexture2D* DerivedTexture);
 	bool WriteDerivedTexture(UTexture2D* DerivedTexture, FString BaseTexName);
@@ -89,4 +109,9 @@ protected:
 	void WriteDerivedMaterials();
 	void UpdateMaterialTextureParameters(UMaterialInstanceConstant* Material, FMaterialInfo& DerivedMaterialInfo);
 	void WriteDerivedStaticMeshAsset();
+
+	void UpdateSourceStaticMeshAsset();
+
+	void UpdateMaterialTextureParameters(UMaterialInstanceDynamic* Material, const FMaterialInfo& SourceMaterialInfo, 
+		const TMap<UTexture2D*,UTexture2D*>& PreviewTextures, UTexture2D* PreviewNormalMap);
 };
\ No newline at end of file
diff --git a/Engine/Plugins/Experimental/MeshLODToolset/Source/MeshLODToolset/Public/Tools/GenerateStaticMeshLODAssetTool.h b/Engine/Plugins/Experimental/MeshLODToolset/Source/MeshLODToolset/Public/Tools/GenerateStaticMeshLODAssetTool.h
--- a/Engine/Plugins/Experimental/MeshLODToolset/Source/MeshLODToolset/Public/Tools/GenerateStaticMeshLODAssetTool.h
+++ b/Engine/Plugins/Experimental/MeshLODToolset/Source/MeshLODToolset/Public/Tools/GenerateStaticMeshLODAssetTool.h
@@ -7,6 +7,7 @@
 #include "MultiSelectionTool.h"
 #include "InteractiveToolBuilder.h"
 #include "DynamicMesh3.h"
+#include "PreviewMesh.h"
 #include "Templates/PimplPtr.h"
 #include "GenerateStaticMeshLODAssetTool.generated.h"
 
@@ -18,6 +19,13 @@ class IAssetGenerationAPI;
 class FGenerateStaticMeshLODProcess;
 
 
+UENUM()
+enum class EGenerateLODAssetOutputMode : uint8
+{
+	UpdateExistingAsset = 0,
+	CreateNewAsset = 1
+};
+
 
 /**
  *
@@ -44,6 +52,9 @@ class MESHLODTOOLSET_API UGenerateStaticMeshLODAssetToolProperties : public UInt
 	GENERATED_BODY()
 
 public:
+	UPROPERTY(EditAnywhere, Category = AssetOptions, meta = (TransientToolProperty))
+	EGenerateLODAssetOutputMode OutputMode = EGenerateLODAssetOutputMode::UpdateExistingAsset;
+
 	UPROPERTY(EditAnywhere, Category = AssetOptions, meta = (TransientToolProperty))
 	FString OutputName;
 
@@ -51,9 +62,8 @@ public:
 	UPROPERTY(EditAnywhere, Category = AssetOptions, meta = (TransientToolProperty))
 	FString GeneratedSuffix;
 
-	/** Name of asset that will be updated */
-	UPROPERTY(VisibleAnywhere, Category = AssetOptions, meta = (TransientToolProperty))
-	FString OutputAsset;
+	UPROPERTY(VisibleAnywhere, Category = Previews)
+	TArray<UTexture2D*> PreviewTextures;
 };
 
 
@@ -85,12 +95,24 @@ protected:
 	UPROPERTY()
 	UGenerateStaticMeshLODAssetToolProperties* BasicProperties;
 
+	UPROPERTY()
+	UPreviewMesh* PreviewMesh;
+
+	UPROPERTY()
+	TArray<UTexture2D*> PreviewTextures;
+
+	UPROPERTY()
+	TArray<UMaterialInterface*> PreviewMaterials;;
+
 protected:
 	UWorld* TargetWorld;
 	IAssetGenerationAPI* AssetAPI;
 
 	TPimplPtr<FGenerateStaticMeshLODProcess> GenerateProcess;
 
+	bool bPreviewValid;
+	void ValidatePreview();
+
 	void CreateNewAsset();
 	void UpdateExistingAsset();
 };
