MeshLODToolset: FGenerateMeshLODGraph now has API for updating parameters on various internal nodes. Added FGenerateStaticMeshLODProcessSettings which contains subset of those parameters that we want to expose to user. Make that available via UGenerateStaticMeshLODAssetToolProperties in AutoGenLODTool. Add necessary watchers and recompute notifications. Add collision data visualization.
#rnx

diff --git a/Engine/Plugins/Experimental/GeometryFlow/Source/GeometryFlowCore/Public/GeometryFlowExecutor.h b/Engine/Plugins/Experimental/GeometryFlow/Source/GeometryFlowCore/Public/GeometryFlowExecutor.h
--- a/Engine/Plugins/Experimental/GeometryFlow/Source/GeometryFlowCore/Public/GeometryFlowExecutor.h
+++ b/Engine/Plugins/Experimental/GeometryFlow/Source/GeometryFlowCore/Public/GeometryFlowExecutor.h
@@ -1,3 +1,4 @@
+// Copyright Epic Games, Inc. All Rights Reserved.
 
 #pragma once
 
diff --git a/Engine/Plugins/Experimental/MeshLODToolset/Source/MeshLODToolset/MeshLODToolset.Build.cs b/Engine/Plugins/Experimental/MeshLODToolset/Source/MeshLODToolset/MeshLODToolset.Build.cs
--- a/Engine/Plugins/Experimental/MeshLODToolset/Source/MeshLODToolset/MeshLODToolset.Build.cs
+++ b/Engine/Plugins/Experimental/MeshLODToolset/Source/MeshLODToolset/MeshLODToolset.Build.cs
@@ -36,6 +36,7 @@ public class MeshLODToolset : ModuleRules
                 "StaticMeshDescription",
 				"ModelingComponents",
 				"ModelingOperators",
+				"MeshModelingTools",
 				"GeometryFlowCore",
 				"GeometryFlowMeshProcessing",
 				"GeometryFlowMeshProcessingEditor"
diff --git a/Engine/Plugins/Experimental/MeshLODToolset/Source/MeshLODToolset/Private/Graphs/GenerateMeshLODGraph.cpp b/Engine/Plugins/Experimental/MeshLODToolset/Source/MeshLODToolset/Private/Graphs/GenerateMeshLODGraph.cpp
--- a/Engine/Plugins/Experimental/MeshLODToolset/Source/MeshLODToolset/Private/Graphs/GenerateMeshLODGraph.cpp
+++ b/Engine/Plugins/Experimental/MeshLODToolset/Source/MeshLODToolset/Private/Graphs/GenerateMeshLODGraph.cpp
@@ -109,6 +109,44 @@ void FGenerateMeshLODGraph::EvaluateResultParallel(
 }
 
 
+void FGenerateMeshLODGraph::UpdateSolidifySettings(const FMeshSolidifySettings& SolidifySettings)
+{
+	UpdateSettingsSourceNodeValue(*Graph, SolidifySettingsNode, SolidifySettings);
+	CurrentSolidifySettings = SolidifySettings;
+}
+
+void FGenerateMeshLODGraph::UpdateMorphologySettings(const FVoxClosureSettings& MorphologySettings)
+{
+	UpdateSettingsSourceNodeValue(*Graph, MorphologySettingsNode, MorphologySettings);
+	CurrentMorphologySettings = MorphologySettings;
+}
+
+void FGenerateMeshLODGraph::UpdateSimplifySettings(const FMeshSimplifySettings& SimplifySettings)
+{
+	UpdateSettingsSourceNodeValue(*Graph, SimplifySettingsNode, SimplifySettings);
+	CurrentSimplifySettings = SimplifySettings;
+}
+
+void FGenerateMeshLODGraph::UpdateAutoUVSettings(const UE::GeometryFlow::FMeshAutoGenerateUVsSettings& AutoUVSettings)
+{
+	UpdateSettingsSourceNodeValue(*Graph, AutoUVSettingsNode, AutoUVSettings);
+	CurrentAutoUVSettings = AutoUVSettings;
+}
+
+void FGenerateMeshLODGraph::UpdateBakeCacheSettings(const UE::GeometryFlow::FMeshMakeBakingCacheSettings& BakeCacheSettings)
+{
+	UpdateSettingsSourceNodeValue(*Graph, BakeCacheSettingsNode, BakeCacheSettings);
+	CurrentBakeCacheSettings = BakeCacheSettings;
+}
+
+
+void FGenerateMeshLODGraph::UpdateGenerateConvexCollisionSettings(const FGenerateConvexHullsCollisionSettings& GenConvexesSettings)
+{
+	UpdateSettingsSourceNodeValue(*Graph, GenerateConvexesSettingsNode, GenConvexesSettings);
+	CurrentGenerateConvexHullsSettings = GenConvexesSettings;
+}
+
+
 void FGenerateMeshLODGraph::EvaluateResult(
 	FDynamicMesh3& ResultMesh,
 	FMeshTangentsd& ResultTangents,
@@ -205,39 +243,39 @@ void FGenerateMeshLODGraph::BuildGraph()
 
 	SolidifyNode = Graph->AddNodeOfType<FSolidifyMeshNode>(TEXT("Solidify"));
 	Graph->InferConnection(MeshSourceNode, SolidifyNode);
-	FGraph::FHandle SolidifySettingsNode = Graph->AddNodeOfType<FSolidifySettingsSourceNode>(TEXT("SolidifySettings"));
+	SolidifySettingsNode = Graph->AddNodeOfType<FSolidifySettingsSourceNode>(TEXT("SolidifySettings"));
 	Graph->InferConnection(SolidifySettingsNode, SolidifyNode);
 
 	MorphologyNode = Graph->AddNodeOfType<FVoxClosureMeshNode>(TEXT("Closure"));
 	Graph->InferConnection(SolidifyNode, MorphologyNode);
-	FGraph::FHandle MorphologySettingsNode = Graph->AddNodeOfType<FVoxClosureSettingsSourceNode>(TEXT("ClosureSettings"));
+	MorphologySettingsNode = Graph->AddNodeOfType<FVoxClosureSettingsSourceNode>(TEXT("ClosureSettings"));
 	Graph->InferConnection(MorphologySettingsNode, MorphologyNode);
 
 	SimplifyNode = Graph->AddNodeOfType<FSimplifyMeshNode>(TEXT("Simplify"));
 	Graph->InferConnection(MorphologyNode, SimplifyNode);
-	FGraph::FHandle SimplifySettingsNode = Graph->AddNodeOfType<FSimplifySettingsSourceNode>(TEXT("SimplifySettings"));
+	SimplifySettingsNode = Graph->AddNodeOfType<FSimplifySettingsSourceNode>(TEXT("SimplifySettings"));
 	Graph->InferConnection(SimplifySettingsNode, SimplifyNode);
 
 	NormalsNode = Graph->AddNodeOfType<FComputeMeshNormalsNode>(TEXT("Normals"));
 	Graph->InferConnection(SimplifyNode, NormalsNode);
-	FGraph::FHandle NormalsSettingsNode = Graph->AddNodeOfType<FNormalsSettingsSourceNode>(TEXT("NormalsSettings"));
+	NormalsSettingsNode = Graph->AddNodeOfType<FNormalsSettingsSourceNode>(TEXT("NormalsSettings"));
 	Graph->InferConnection(NormalsSettingsNode, NormalsNode);
 
 	// computing UVs
 
 	AutoUVNode = Graph->AddNodeOfType<FMeshAutoGenerateUVsNode>(TEXT("AutoUV"));
 	Graph->InferConnection(NormalsNode, AutoUVNode);
-	FGraph::FHandle AutoUVSettingsNode = Graph->AddNodeOfType<FMeshAutoGenerateUVsSettingsSourceNode>(TEXT("AutoUVSettings"));
+	AutoUVSettingsNode = Graph->AddNodeOfType<FMeshAutoGenerateUVsSettingsSourceNode>(TEXT("AutoUVSettings"));
 	Graph->InferConnection(AutoUVSettingsNode, AutoUVNode);
 
 	RecomputeUVNode = Graph->AddNodeOfType<FMeshRecalculateUVsNode>(TEXT("RecalcUV"));
 	Graph->InferConnection(AutoUVNode, RecomputeUVNode);
-	FGraph::FHandle RecomputeUVSettingsNode = Graph->AddNodeOfType<FMeshRecalculateUVsSettingsSourceNode>(TEXT("RecalcUVSettings"));
+	RecomputeUVSettingsNode = Graph->AddNodeOfType<FMeshRecalculateUVsSettingsSourceNode>(TEXT("RecalcUVSettings"));
 	Graph->InferConnection(RecomputeUVSettingsNode, RecomputeUVNode);
 
 	RepackUVNode = Graph->AddNodeOfType<FMeshRepackUVsNode>(TEXT("RepackUV"));
 	Graph->InferConnection(RecomputeUVNode, RepackUVNode);
-	FGraph::FHandle RepackUVSettingsNode = Graph->AddNodeOfType<FMeshRepackUVsSettingsSourceNode>(TEXT("RepackUVSettings"));
+	RepackUVSettingsNode = Graph->AddNodeOfType<FMeshRepackUVsSettingsSourceNode>(TEXT("RepackUVSettings"));
 	Graph->InferConnection(RepackUVSettingsNode, RepackUVNode);
 
 
@@ -251,7 +289,7 @@ void FGenerateMeshLODGraph::BuildGraph()
 
 	TangentsNode = Graph->AddNodeOfType<FComputeMeshTangentsNode>(TEXT("Tangents"));
 	Graph->InferConnection(RepackUVNode, TangentsNode);
-	FGraph::FHandle TangentsSettingsNode = Graph->AddNodeOfType<FTangentsSettingsSourceNode>(TEXT("TangentsSettings"));
+	TangentsSettingsNode = Graph->AddNodeOfType<FTangentsSettingsSourceNode>(TEXT("TangentsSettings"));
 	Graph->InferConnection(TangentsSettingsNode, TangentsNode);
 
 	// tangents output
@@ -263,7 +301,7 @@ void FGenerateMeshLODGraph::BuildGraph()
 	BakeCacheNode = Graph->AddNodeOfType<FMakeMeshBakingCacheNode>(TEXT("MakeBakeCache"));
 	Graph->AddConnection(MeshSourceNode, FDynamicMeshSourceNode::OutParamValue(), BakeCacheNode, FMakeMeshBakingCacheNode::InParamDetailMesh());
 	Graph->AddConnection(MeshOutputNode, FDynamicMeshTransferNode::OutParamValue(), BakeCacheNode, FMakeMeshBakingCacheNode::InParamTargetMesh());
-	FGraph::FHandle BakeCacheSettingsNode = Graph->AddNodeOfType<FMeshMakeBakingCacheSettingsSourceNode>(TEXT("BakeCacheSettings"));
+	BakeCacheSettingsNode = Graph->AddNodeOfType<FMeshMakeBakingCacheSettingsSourceNode>(TEXT("BakeCacheSettings"));
 	Graph->InferConnection(BakeCacheSettingsNode, BakeCacheNode);
 
 	// normal map baker
@@ -271,7 +309,7 @@ void FGenerateMeshLODGraph::BuildGraph()
 	BakeNormalMapNode = Graph->AddNodeOfType<FBakeMeshNormalMapNode>(TEXT("BakeNormalMap"));
 	Graph->InferConnection(BakeCacheNode, BakeNormalMapNode);
 	Graph->InferConnection(TangentsNode, BakeNormalMapNode);
-	FGraph::FHandle BakeNormalMapSettingsNode = Graph->AddNodeOfType<FBakeMeshNormalMapSettingsSourceNode>(TEXT("BakeNormalMapSettings"));
+	BakeNormalMapSettingsNode = Graph->AddNodeOfType<FBakeMeshNormalMapSettingsSourceNode>(TEXT("BakeNormalMapSettings"));
 	Graph->InferConnection(BakeNormalMapSettingsNode, BakeNormalMapNode);
 
 
@@ -289,7 +327,7 @@ void FGenerateMeshLODGraph::BuildGraph()
 	GenerateConvexesNode = Graph->AddNodeOfType<FGenerateConvexHullsCollisionNode>(TEXT("GenerateConvexes"));
 	Graph->InferConnection(MeshSourceNode, GenerateConvexesNode);
 	Graph->InferConnection(DecomposeMeshForCollisionNode, GenerateConvexesNode);
-	FGraph::FHandle GenerateConvexesSettingsNode = Graph->AddNodeOfType<FGenerateConvexHullsCollisionSettingsSourceNode>(TEXT("GenerateConvexesSettings"));
+	GenerateConvexesSettingsNode = Graph->AddNodeOfType<FGenerateConvexHullsCollisionSettingsSourceNode>(TEXT("GenerateConvexesSettings"));
 	Graph->InferConnection(GenerateConvexesSettingsNode, GenerateConvexesNode);
 
 	// final collision output
@@ -304,18 +342,18 @@ void FGenerateMeshLODGraph::BuildGraph()
 
 
 	FMeshSolidifySettings SolidifySettings;
-	UpdateSettingsSourceNodeValue(*Graph, SolidifySettingsNode, SolidifySettings);
+	UpdateSolidifySettings(SolidifySettings);
 
 	FVoxClosureSettings MorphologySettings;
 	MorphologySettings.Distance = 5.0;
-	UpdateSettingsSourceNodeValue(*Graph, MorphologySettingsNode, MorphologySettings);
+	UpdateMorphologySettings(MorphologySettings);
 
 	FMeshSimplifySettings SimplifySettings;
 	SimplifySettings.bDiscardAttributes = true;
 	SimplifySettings.SimplifyType = EMeshSimplifyType::VolumePreserving;
 	SimplifySettings.TargetType = EMeshSimplifyTargetType::TriangleCount;
 	SimplifySettings.TargetCount = 500;
-	UpdateSettingsSourceNodeValue(*Graph, SimplifySettingsNode, SimplifySettings);
+	UpdateSimplifySettings(SimplifySettings);
 
 	FMeshNormalsSettings NormalsSettings;
 	NormalsSettings.NormalsType = EComputeNormalsType::FromFaceAngleThreshold;
@@ -323,7 +361,9 @@ void FGenerateMeshLODGraph::BuildGraph()
 	UpdateSettingsSourceNodeValue(*Graph, NormalsSettingsNode, NormalsSettings);
 
 	FMeshAutoGenerateUVsSettings AutoUVSettings;
-	UpdateSettingsSourceNodeValue(*Graph, AutoUVSettingsNode, AutoUVSettings);
+	AutoUVSettings.NumCharts = 20;
+	AutoUVSettings.Stretch = 0.1;
+	UpdateAutoUVSettings(AutoUVSettings);
 
 	FMeshRecalculateUVsSettings RecomputeUVSettings;
 	UpdateSettingsSourceNodeValue(*Graph, RecomputeUVSettingsNode, RecomputeUVSettings);
@@ -339,7 +379,7 @@ void FGenerateMeshLODGraph::BuildGraph()
 	FMeshMakeBakingCacheSettings BakeCacheSettings;
 	BakeCacheSettings.Dimensions = FImageDimensions(512, 512);
 	BakeCacheSettings.Thickness = 5.0;
-	UpdateSettingsSourceNodeValue(*Graph, BakeCacheSettingsNode, BakeCacheSettings);
+	UpdateBakeCacheSettings(BakeCacheSettings);
 
 
 	FBakeMeshNormalMapSettings NormalMapSettings;
@@ -350,8 +390,7 @@ void FGenerateMeshLODGraph::BuildGraph()
 	UpdateSettingsSourceNodeValue(*Graph, IgnoreGroupsForCollisionNode, IgnoreGroupsForCollision);
 
 	FGenerateConvexHullsCollisionSettings GenConvexesSettings;
-	UpdateSettingsSourceNodeValue(*Graph, GenerateConvexesSettingsNode, GenConvexesSettings);
-
+	UpdateGenerateConvexCollisionSettings(GenConvexesSettings);
 }
 
 
diff --git a/Engine/Plugins/Experimental/MeshLODToolset/Source/MeshLODToolset/Private/Graphs/GenerateStaticMeshLODProcess.cpp b/Engine/Plugins/Experimental/MeshLODToolset/Source/MeshLODToolset/Private/Graphs/GenerateStaticMeshLODProcess.cpp
--- a/Engine/Plugins/Experimental/MeshLODToolset/Source/MeshLODToolset/Private/Graphs/GenerateStaticMeshLODProcess.cpp
+++ b/Engine/Plugins/Experimental/MeshLODToolset/Source/MeshLODToolset/Private/Graphs/GenerateStaticMeshLODProcess.cpp
@@ -176,6 +176,7 @@ bool FGenerateStaticMeshLODProcess::InitializeGenerator()
 	Generator = MakeUnique<FGenerateMeshLODGraph>();
 	Generator->BuildGraph();
 
+	// initialize source textures
 	TextureToDerivedTexIndex.Reset();
 	for (const FMaterialInfo& MatInfo : SourceMaterials)
 	{
@@ -189,15 +190,96 @@ bool FGenerateStaticMeshLODProcess::InitializeGenerator()
 			}
 		}
 	}
+	
 
+	// initialize source mesh
 	Generator->SetSourceMesh(this->SourceMesh);
 
+
+	// read back default settings
+
+	CurrentSettings.SolidifyVoxelResolution = Generator->GetCurrentSolidifySettings().VoxelResolution;
+	CurrentSettings.WindingThreshold = Generator->GetCurrentSolidifySettings().WindingThreshold;
+
+	//CurrentSettings.MorphologyVoxelResolution = Generator->GetCurrentMorphologySettings().VoxelResolution;
+	CurrentSettings.ClosureDistance = Generator->GetCurrentMorphologySettings().Distance;
+
+	CurrentSettings.SimplifyTriangleCount = Generator->GetCurrentSimplifySettings().TargetCount;
+
+	CurrentSettings.BakeResolution = (EGenerateStaticMeshLODBakeResolution)Generator->GetCurrentBakeCacheSettings().Dimensions.GetWidth();
+	CurrentSettings.BakeThickness = Generator->GetCurrentBakeCacheSettings().Thickness;
+
+	CurrentSettings.ConvexTriangleCount = Generator->GetCurrentGenerateConvexCollisionSettings().SimplifyToTriangleCount;
+
+
 	return true;
 }
 
 
+void FGenerateStaticMeshLODProcess::UpdateSettings(const FGenerateStaticMeshLODProcessSettings& NewSettings)
+{
+	bool bSharedVoxelResolutionChanged = (NewSettings.SolidifyVoxelResolution != CurrentSettings.SolidifyVoxelResolution);
+	if ( bSharedVoxelResolutionChanged
+		|| (NewSettings.WindingThreshold != CurrentSettings.WindingThreshold))
+	{
+		UE::GeometryFlow::FMeshSolidifySettings NewSolidifySettings = Generator->GetCurrentSolidifySettings();
+		NewSolidifySettings.VoxelResolution = NewSettings.SolidifyVoxelResolution;
+		NewSolidifySettings.WindingThreshold = NewSettings.WindingThreshold;
+		Generator->UpdateSolidifySettings(NewSolidifySettings);
+	}
+
+
+	if ( bSharedVoxelResolutionChanged
+		|| (NewSettings.ClosureDistance != CurrentSettings.ClosureDistance))
+	{
+		UE::GeometryFlow::FVoxClosureSettings NewClosureSettings = Generator->GetCurrentMorphologySettings();
+		NewClosureSettings.VoxelResolution = NewSettings.SolidifyVoxelResolution;
+		NewClosureSettings.Distance = NewSettings.ClosureDistance;
+		Generator->UpdateMorphologySettings(NewClosureSettings);
+	}
+
+
+	if (NewSettings.SimplifyTriangleCount != CurrentSettings.SimplifyTriangleCount)
+	{
+		UE::GeometryFlow::FMeshSimplifySettings NewSimplifySettings = Generator->GetCurrentSimplifySettings();
+		NewSimplifySettings.TargetCount = NewSettings.SimplifyTriangleCount;
+		Generator->UpdateSimplifySettings(NewSimplifySettings);
+	}
+
+	if (NewSettings.NumAutoUVCharts != CurrentSettings.NumAutoUVCharts)
+	{
+		UE::GeometryFlow::FMeshAutoGenerateUVsSettings NewAutoUVSettings = Generator->GetCurrentAutoUVSettings();
+		NewAutoUVSettings.NumCharts = NewSettings.NumAutoUVCharts;
+		Generator->UpdateAutoUVSettings(NewAutoUVSettings);
+	}
+
+
+	if ( (NewSettings.BakeResolution != CurrentSettings.BakeResolution) ||
+		 (NewSettings.BakeThickness != CurrentSettings.BakeThickness))
+	{
+		UE::GeometryFlow::FMeshMakeBakingCacheSettings NewBakeSettings = Generator->GetCurrentBakeCacheSettings();
+		NewBakeSettings.Dimensions = FImageDimensions((int32)NewSettings.BakeResolution, (int32)NewSettings.BakeResolution);
+		NewBakeSettings.Thickness = NewSettings.BakeThickness;
+		Generator->UpdateBakeCacheSettings(NewBakeSettings);
+	}
+
+
+	if (NewSettings.ConvexTriangleCount != CurrentSettings.ConvexTriangleCount)
+	{
+		UE::GeometryFlow::FGenerateConvexHullsCollisionSettings NewGenConvexSettings = Generator->GetCurrentGenerateConvexCollisionSettings();
+		NewGenConvexSettings.SimplifyToTriangleCount = NewSettings.ConvexTriangleCount;
+		Generator->UpdateGenerateConvexCollisionSettings(NewGenConvexSettings);
+	}
+
+	CurrentSettings = NewSettings;
+}
+
+
+
+
 bool FGenerateStaticMeshLODProcess::ComputeDerivedSourceData()
 {
+	DerivedTextureImages.Reset();
 	if (bUseParallelExecutor)
 	{
 		Generator->EvaluateResultParallel(
@@ -217,6 +299,7 @@ bool FGenerateStaticMeshLODProcess::ComputeDerivedSourceData()
 			this->DerivedTextureImages);
 	}
 
+
 	// copy all materials for now...we are going to replace all the images though, and
 	// should not copy those?
 	DerivedMaterials = SourceMaterials;
diff --git a/Engine/Plugins/Experimental/MeshLODToolset/Source/MeshLODToolset/Private/Tools/GenerateStaticMeshLODAssetTool.cpp b/Engine/Plugins/Experimental/MeshLODToolset/Source/MeshLODToolset/Private/Tools/GenerateStaticMeshLODAssetTool.cpp
--- a/Engine/Plugins/Experimental/MeshLODToolset/Source/MeshLODToolset/Private/Tools/GenerateStaticMeshLODAssetTool.cpp
+++ b/Engine/Plugins/Experimental/MeshLODToolset/Source/MeshLODToolset/Private/Tools/GenerateStaticMeshLODAssetTool.cpp
@@ -21,6 +21,9 @@
 #include "AssetGenerationUtil.h"
 #include "Selection/ToolSelectionUtil.h"
 
+#include "Physics/PhysicsDataCollection.h"
+#include "Physics/CollisionGeometryVisualization.h"
+
 #include "Components/StaticMeshComponent.h"
 #include "Engine/StaticMesh.h"
 
@@ -110,7 +113,7 @@ void UGenerateStaticMeshLODAssetTool::Setup()
 		LOCTEXT("OnStartStaticMeshLODAssetTool", "This tool creates a new LOD asset"),
 		EToolMessageLevel::UserNotification);
 
-	GenerateProcess = MakePimpl<FGenerateStaticMeshLODProcess>();
+	GenerateProcess = MakeUnique<FGenerateStaticMeshLODProcess>();
 
 	TUniquePtr<FPrimitiveComponentTarget>& SourceComponent = ComponentTargets[0];
 	UStaticMeshComponent* StaticMeshComponent = CastChecked<UStaticMeshComponent>(SourceComponent->GetOwnerComponent());
@@ -120,21 +123,44 @@ void UGenerateStaticMeshLODAssetTool::Setup()
 		GenerateProcess->Initialize(StaticMesh);
 	}
 
+	BasicProperties->GeneratorSettings = GenerateProcess->GetCurrentSettings();
+	BasicProperties->WatchProperty(BasicProperties->GeneratorSettings.SolidifyVoxelResolution, [this](int) { OnSettingsModified(); });
+	BasicProperties->WatchProperty(BasicProperties->GeneratorSettings.WindingThreshold, [this](float) { OnSettingsModified(); });
+	//BasicProperties->WatchProperty(BasicProperties->GeneratorSettings.MorphologyVoxelResolution, [this](int) { OnSettingsModified(); });
+	BasicProperties->WatchProperty(BasicProperties->GeneratorSettings.ClosureDistance, [this](float) { OnSettingsModified(); });
+	BasicProperties->WatchProperty(BasicProperties->GeneratorSettings.SimplifyTriangleCount, [this](int) { OnSettingsModified(); });
+	BasicProperties->WatchProperty(BasicProperties->GeneratorSettings.NumAutoUVCharts, [this](int) { OnSettingsModified(); });
+	BasicProperties->WatchProperty(BasicProperties->GeneratorSettings.BakeResolution, [this](EGenerateStaticMeshLODBakeResolution) { OnSettingsModified(); });
+	BasicProperties->WatchProperty(BasicProperties->GeneratorSettings.BakeThickness, [this](float) { OnSettingsModified(); });
+	BasicProperties->WatchProperty(BasicProperties->GeneratorSettings.ConvexTriangleCount, [this](int) { OnSettingsModified(); });
+
+
+
 	FBoxSphereBounds Bounds = StaticMeshComponent->Bounds;
 	FTransform PreviewTransform = SourceComponent->GetWorldTransform();
-	PreviewTransform.AddToTranslation(FVector(0, 2.0f*Bounds.SphereRadius, 0));
+	PreviewTransform.AddToTranslation(FVector(0, 2.5f*Bounds.BoxExtent.Y, 0));
 
 	PreviewMesh = NewObject<UPreviewMesh>(this);
 	PreviewMesh->CreateInWorld(TargetWorld, PreviewTransform);
 	PreviewMesh->SetVisible(true);
 	PreviewMesh->SetTangentsMode(EDynamicMeshTangentCalcType::ExternallyCalculated);
 
+	CollisionVizSettings = NewObject<UCollisionGeometryVisualizationProperties>(this);
+	CollisionVizSettings->RestoreProperties(this);
+	AddToolPropertySource(CollisionVizSettings);
+	CollisionVizSettings->WatchProperty(CollisionVizSettings->LineThickness, [this](float NewValue) { bCollisionVisualizationDirty = true; });
+	CollisionVizSettings->WatchProperty(CollisionVizSettings->Color, [this](FColor NewValue) { bCollisionVisualizationDirty = true; });
+	CollisionVizSettings->WatchProperty(CollisionVizSettings->bShowHidden, [this](bool bNewValue) { bCollisionVisualizationDirty = true; });
+
+	CollisionPreview = NewObject<UPreviewGeometry>(this);
+	CollisionPreview->CreateInWorld(TargetWorld, PreviewTransform);
+
 	// Recompute if we switch between parallel and serial
 	int32 WatcherIndex = BasicProperties->WatchProperty(BasicProperties->bParallelExecution, [this](bool bNewParallelExec)
 	{
 		// TODO: We crash if we don't recreate the Process and reinitialize it. Why?
 
-		GenerateProcess = MakePimpl<FGenerateStaticMeshLODProcess>();
+		GenerateProcess = MakeUnique<FGenerateStaticMeshLODProcess>();
 
 		TUniquePtr<FPrimitiveComponentTarget>& SourceComponent = ComponentTargets[0];
 		UStaticMeshComponent* StaticMeshComponent = CastChecked<UStaticMeshComponent>(SourceComponent->GetOwnerComponent());
@@ -145,24 +171,35 @@ void UGenerateStaticMeshLODAssetTool::Setup()
 		}
 
 		bPreviewValid = false;
-		ValidatePreview();
+		//ValidatePreview();
 	});
 	BasicProperties->SilentUpdateWatcherAtIndex(WatcherIndex);
 
 	bPreviewValid = false;
-	ValidatePreview();
 }
 
 
+void UGenerateStaticMeshLODAssetTool::OnSettingsModified()
+{
+	UE_LOG(LogTemp, Warning, TEXT("SETTINGS MODIFIED!"));
+
+	GenerateProcess->UpdateSettings(BasicProperties->GeneratorSettings);
+	bPreviewValid = false;
+}
+
 
 void UGenerateStaticMeshLODAssetTool::Shutdown(EToolShutdownType ShutdownType)
 {
 	BasicProperties->SaveProperties(this);
+	CollisionVizSettings->SaveProperties(this);
 
 	PreviewMesh->SetVisible(false);
 	PreviewMesh->Disconnect();
 	PreviewMesh = nullptr;
 
+	CollisionPreview->Disconnect();
+	CollisionPreview = nullptr;
+
 	if (ShutdownType == EToolShutdownType::Accept)
 	{
 		if (BasicProperties->OutputMode == EGenerateLODAssetOutputMode::UpdateExistingAsset)
@@ -182,14 +219,22 @@ void UGenerateStaticMeshLODAssetTool::SetAssetAPI(IAssetGenerationAPI* AssetAPII
 	this->AssetAPI = AssetAPIIn;
 }
 
-bool UGenerateStaticMeshLODAssetTool::HasAccept() const
+
+bool UGenerateStaticMeshLODAssetTool::CanAccept() const
 {
 	return true;
 }
 
-bool UGenerateStaticMeshLODAssetTool::CanAccept() const
+
+void UGenerateStaticMeshLODAssetTool::OnTick(float DeltaTime)
 {
-	return true;
+	ValidatePreview();
+
+	if (bCollisionVisualizationDirty)
+	{
+		UpdateCollisionVisualization();
+		bCollisionVisualizationDirty = false;
+	}
 }
 
 void UGenerateStaticMeshLODAssetTool::ValidatePreview()
@@ -201,7 +246,8 @@ void UGenerateStaticMeshLODAssetTool::ValidatePreview()
 	GenerateProcess->ComputeDerivedSourceData();
 	const FDynamicMesh3& ResultMesh = GenerateProcess->GetDerivedLOD0Mesh();
 	const FMeshTangentsd& ResultTangents = GenerateProcess->GetDerivedLOD0MeshTangents();
-	
+	const FSimpleShapeSet3d& ResultCollision = GenerateProcess->GetDerivedCollision();
+
 	PreviewMesh->EditMesh([&](FDynamicMesh3& MeshToUpdate)
 	{
 		MeshToUpdate = GenerateProcess->GetDerivedLOD0Mesh();
@@ -219,10 +265,33 @@ void UGenerateStaticMeshLODAssetTool::ValidatePreview()
 		BasicProperties->PreviewTextures = PreviewTextures;
 	}
 
+
+	FPhysicsDataCollection PhysicsData;
+	PhysicsData.Geometry = ResultCollision;
+	PhysicsData.CopyGeometryToAggregate();
+	UE::PhysicsTools::InitializePreviewGeometryLines(PhysicsData, CollisionPreview,
+		CollisionVizSettings->Color, CollisionVizSettings->LineThickness, 0.0f, 16);
+
 	bPreviewValid = true;
 }
 
 
+
+void UGenerateStaticMeshLODAssetTool::UpdateCollisionVisualization()
+{
+	float UseThickness = CollisionVizSettings->LineThickness;
+	FColor UseColor = CollisionVizSettings->Color;
+	LineMaterial = ToolSetupUtil::GetDefaultLineComponentMaterial(GetToolManager(), !CollisionVizSettings->bShowHidden);
+
+	CollisionPreview->UpdateAllLineSets([&](ULineSetComponent* LineSet)
+	{
+		LineSet->SetAllLinesThickness(UseThickness);
+		LineSet->SetAllLinesColor(UseColor);
+	});
+	CollisionPreview->SetAllLineSetsMaterial(LineMaterial);
+}
+
+
 void UGenerateStaticMeshLODAssetTool::CreateNewAsset()
 {
 	GenerateProcess->CalculateDerivedPathName(BasicProperties->GeneratedSuffix);
diff --git a/Engine/Plugins/Experimental/MeshLODToolset/Source/MeshLODToolset/Public/Graphs/GenerateMeshLODGraph.h b/Engine/Plugins/Experimental/MeshLODToolset/Source/MeshLODToolset/Public/Graphs/GenerateMeshLODGraph.h
--- a/Engine/Plugins/Experimental/MeshLODToolset/Source/MeshLODToolset/Public/Graphs/GenerateMeshLODGraph.h
+++ b/Engine/Plugins/Experimental/MeshLODToolset/Source/MeshLODToolset/Public/Graphs/GenerateMeshLODGraph.h
@@ -9,6 +9,14 @@
 #include "DynamicMesh3.h"
 #include "MeshTangents.h"
 
+#include "MeshProcessingNodes/MeshSolidifyNode.h"
+#include "MeshProcessingNodes/MeshVoxMorphologyNode.h"
+#include "MeshProcessingNodes/MeshSimplifyNode.h"
+#include "MeshProcessingNodes/MeshDeleteTrianglesNode.h"
+#include "MeshProcessingNodes/MeshAutoGenerateUVsNode.h"
+#include "DataTypes/MeshImageBakingData.h"
+#include "PhysicsNodes/GenerateConvexHullsCollisionNode.h"
+
 
 class FGenerateMeshLODGraph
 {
@@ -20,6 +28,28 @@ public:
 	void SetSourceMesh(const FDynamicMesh3& SourceMesh);
 
 
+	void UpdateSolidifySettings(const UE::GeometryFlow::FMeshSolidifySettings& SolidifySettings);
+	const UE::GeometryFlow::FMeshSolidifySettings& GetCurrentSolidifySettings() const { return CurrentSolidifySettings; }
+
+	void UpdateMorphologySettings(const UE::GeometryFlow::FVoxClosureSettings& MorphologySettings);
+	const UE::GeometryFlow::FVoxClosureSettings& GetCurrentMorphologySettings() const { return CurrentMorphologySettings; }
+
+	void UpdateSimplifySettings(const UE::GeometryFlow::FMeshSimplifySettings& SimplifySettings);
+	const UE::GeometryFlow::FMeshSimplifySettings& GetCurrentSimplifySettings() const { return CurrentSimplifySettings; }
+
+	void UpdateAutoUVSettings(const UE::GeometryFlow::FMeshAutoGenerateUVsSettings& AutoUVSettings);
+	const UE::GeometryFlow::FMeshAutoGenerateUVsSettings& GetCurrentAutoUVSettings() const { return CurrentAutoUVSettings; }
+
+	void UpdateBakeCacheSettings(const UE::GeometryFlow::FMeshMakeBakingCacheSettings& BakeCacheSettings);
+	const UE::GeometryFlow::FMeshMakeBakingCacheSettings& GetCurrentBakeCacheSettings() const { return CurrentBakeCacheSettings; }
+
+
+	void UpdateGenerateConvexCollisionSettings(const UE::GeometryFlow::FGenerateConvexHullsCollisionSettings& ConvexesSettings);
+	const UE::GeometryFlow::FGenerateConvexHullsCollisionSettings& GetCurrentGenerateConvexCollisionSettings() const { return CurrentGenerateConvexHullsSettings; }
+
+
+
+
 	void EvaluateResult(
 		FDynamicMesh3& ResultMesh,
 		FMeshTangentsd& ResultTangents,
@@ -40,20 +70,39 @@ protected:
 	UE::GeometryFlow::FGraph::FHandle MeshSourceNode;
 
 	UE::GeometryFlow::FGraph::FHandle SolidifyNode;
+	UE::GeometryFlow::FGraph::FHandle SolidifySettingsNode;
+	UE::GeometryFlow::FMeshSolidifySettings CurrentSolidifySettings;
+
 	UE::GeometryFlow::FGraph::FHandle MorphologyNode;
+	UE::GeometryFlow::FGraph::FHandle MorphologySettingsNode;
+	UE::GeometryFlow::FVoxClosureSettings CurrentMorphologySettings;
 
 	UE::GeometryFlow::FGraph::FHandle SimplifyNode;
+	UE::GeometryFlow::FGraph::FHandle SimplifySettingsNode;
+	UE::GeometryFlow::FMeshSimplifySettings CurrentSimplifySettings;
 
 	UE::GeometryFlow::FGraph::FHandle NormalsNode;
+	UE::GeometryFlow::FGraph::FHandle NormalsSettingsNode;
 
 	UE::GeometryFlow::FGraph::FHandle AutoUVNode;
+	UE::GeometryFlow::FGraph::FHandle AutoUVSettingsNode;
+	UE::GeometryFlow::FMeshAutoGenerateUVsSettings CurrentAutoUVSettings;
+
 	UE::GeometryFlow::FGraph::FHandle RecomputeUVNode;
+	UE::GeometryFlow::FGraph::FHandle RecomputeUVSettingsNode;
+
 	UE::GeometryFlow::FGraph::FHandle RepackUVNode;
+	UE::GeometryFlow::FGraph::FHandle RepackUVSettingsNode;
 
 	UE::GeometryFlow::FGraph::FHandle TangentsNode;
+	UE::GeometryFlow::FGraph::FHandle TangentsSettingsNode;
 
 	UE::GeometryFlow::FGraph::FHandle BakeCacheNode;
+	UE::GeometryFlow::FGraph::FHandle BakeCacheSettingsNode;
+	UE::GeometryFlow::FMeshMakeBakingCacheSettings CurrentBakeCacheSettings;
+
 	UE::GeometryFlow::FGraph::FHandle BakeNormalMapNode;
+	UE::GeometryFlow::FGraph::FHandle BakeNormalMapSettingsNode;
 
 	struct FBakeTextureGraphInfo
 	{
@@ -65,7 +114,10 @@ protected:
 	TArray<FBakeTextureGraphInfo> BakeTextureNodes;
 
 	UE::GeometryFlow::FGraph::FHandle DecomposeMeshForCollisionNode;
+	
 	UE::GeometryFlow::FGraph::FHandle GenerateConvexesNode;
+	UE::GeometryFlow::FGraph::FHandle GenerateConvexesSettingsNode;
+	UE::GeometryFlow::FGenerateConvexHullsCollisionSettings CurrentGenerateConvexHullsSettings;
 
 	UE::GeometryFlow::FGraph::FHandle CollisionOutputNode;
 	UE::GeometryFlow::FGraph::FHandle MeshOutputNode;
diff --git a/Engine/Plugins/Experimental/MeshLODToolset/Source/MeshLODToolset/Public/Graphs/GenerateStaticMeshLODProcess.h b/Engine/Plugins/Experimental/MeshLODToolset/Source/MeshLODToolset/Public/Graphs/GenerateStaticMeshLODProcess.h
--- a/Engine/Plugins/Experimental/MeshLODToolset/Source/MeshLODToolset/Public/Graphs/GenerateStaticMeshLODProcess.h
+++ b/Engine/Plugins/Experimental/MeshLODToolset/Source/MeshLODToolset/Public/Graphs/GenerateStaticMeshLODProcess.h
@@ -9,16 +9,92 @@
 
 #include "Graphs/GenerateMeshLODGraph.h"
 
+#include "GenerateStaticMeshLODProcess.generated.h"
+
+
 class UTexture2D;
 class UMaterialInstanceConstant;
 class UMaterialInstanceDynamic;
 
+
+UENUM()
+enum class EGenerateStaticMeshLODBakeResolution
+{
+	Resolution16 = 16 UMETA(DisplayName = "16 x 16"),
+	Resolution32 = 32 UMETA(DisplayName = "32 x 32"),
+	Resolution64 = 64 UMETA(DisplayName = "64 x 64"),
+	Resolution128 = 128 UMETA(DisplayName = "128 x 128"),
+	Resolution256 = 256 UMETA(DisplayName = "256 x 256"),
+	Resolution512 = 512 UMETA(DisplayName = "512 x 512"),
+	Resolution1024 = 1024 UMETA(DisplayName = "1024 x 1024"),
+	Resolution2048 = 2048 UMETA(DisplayName = "2048 x 2048"),
+	Resolution4096 = 4096 UMETA(DisplayName = "4096 x 4096"),
+	Resolution8192 = 8192 UMETA(DisplayName = "8192 x 8192")
+};
+
+
+USTRUCT()
+struct FGenerateStaticMeshLODProcessSettings
+{
+	GENERATED_BODY()
+
+	// Solidify settings
+
+	UPROPERTY(EditAnywhere, Category = Solidify, meta = (DisplayName="Voxel Resolution"))
+	int SolidifyVoxelResolution = 64;
+
+	UPROPERTY(EditAnywhere, Category = Solidify, AdvancedDisplay)
+	float WindingThreshold = 0.5f;
+
+
+	// Morphology settings
+
+	//UPROPERTY(EditAnywhere, Category = Morphology, meta = (DisplayName = "Voxel Resolution"))
+	//int MorphologyVoxelResolution = 64;
+
+	UPROPERTY(EditAnywhere, Category = Morphology, meta = (DisplayName = "Closure Distance"))
+	float ClosureDistance = 1.0f;
+
+
+	// Simplify settings
+
+	UPROPERTY(EditAnywhere, Category = Simplify, meta = (DisplayName = "Simplify Tri Count"))
+	int SimplifyTriangleCount = 500;
+
+
+	// UV settings
+
+	UPROPERTY(EditAnywhere, Category = AutoUV, meta = (DisplayName = "AutoUV Charts", UIMin = 0, UIMax = 1000))
+	int NumAutoUVCharts = 20;
+
+
+	// Bake Settings
+
+	UPROPERTY(EditAnywhere, Category = Baking , meta = (DisplayName = "Bake Image Res"))
+	EGenerateStaticMeshLODBakeResolution BakeResolution = EGenerateStaticMeshLODBakeResolution::Resolution512;
+
+	UPROPERTY(EditAnywhere, Category = Baking, meta = (DisplayName = "Bake Thickness"))
+	float BakeThickness = 5.0f;
+
+
+
+	// Convex Hull Settings
+
+	UPROPERTY(EditAnywhere, Category = ConvexCollision, meta = (DisplayName = "Convex Tri Count"))
+	int ConvexTriangleCount = 50;
+};
+
+
+
 class FGenerateStaticMeshLODProcess
 {
 public:
 
 	bool Initialize(UStaticMesh* SourceMesh);
 
+	const FGenerateStaticMeshLODProcessSettings& GetCurrentSettings() const { return CurrentSettings; }
+	void UpdateSettings(const FGenerateStaticMeshLODProcessSettings& NewSettings);
+
 
 	UStaticMesh* GetSourceStaticMesh() const { return SourceStaticMesh; }
 	const FDynamicMesh3& GetSourceMesh() const { return SourceMesh; }
@@ -33,7 +109,7 @@ public:
 	bool ComputeDerivedSourceData();
 	const FDynamicMesh3& GetDerivedLOD0Mesh() const { return DerivedLODMesh; }
 	const FMeshTangentsd& GetDerivedLOD0MeshTangents() const { return DerivedLODMeshTangents; }
-
+	const FSimpleShapeSet3d& GetDerivedCollision() const { return DerivedCollision; }
 
 	/**
 	 * Creates new Asset
@@ -104,6 +180,8 @@ protected:
 	TUniquePtr<FGenerateMeshLODGraph> Generator;
 	bool InitializeGenerator();
 
+	FGenerateStaticMeshLODProcessSettings CurrentSettings;
+
 	bool WriteDerivedTexture(UTexture2D* SourceTexture, UTexture2D* DerivedTexture);
 	bool WriteDerivedTexture(UTexture2D* DerivedTexture, FString BaseTexName);
 	void WriteDerivedTextures();
diff --git a/Engine/Plugins/Experimental/MeshLODToolset/Source/MeshLODToolset/Public/Tools/GenerateStaticMeshLODAssetTool.h b/Engine/Plugins/Experimental/MeshLODToolset/Source/MeshLODToolset/Public/Tools/GenerateStaticMeshLODAssetTool.h
--- a/Engine/Plugins/Experimental/MeshLODToolset/Source/MeshLODToolset/Public/Tools/GenerateStaticMeshLODAssetTool.h
+++ b/Engine/Plugins/Experimental/MeshLODToolset/Source/MeshLODToolset/Public/Tools/GenerateStaticMeshLODAssetTool.h
@@ -8,7 +8,9 @@
 #include "InteractiveToolBuilder.h"
 #include "DynamicMesh3.h"
 #include "PreviewMesh.h"
-#include "Templates/PimplPtr.h"
+#include "Drawing/PreviewGeometryActor.h"
+#include "Graphs/GenerateStaticMeshLODProcess.h"
+#include "Physics/CollisionPropertySets.h"
 #include "GenerateStaticMeshLODAssetTool.generated.h"
 
 
@@ -16,7 +18,6 @@
 struct FMeshDescription;
 class USimpleDynamicMeshComponent;
 class IAssetGenerationAPI;
-class FGenerateStaticMeshLODProcess;
 
 
 UENUM()
@@ -62,12 +63,16 @@ public:
 	UPROPERTY(EditAnywhere, Category = AssetOptions, meta = (TransientToolProperty))
 	FString GeneratedSuffix;
 
+	UPROPERTY(EditAnywhere, Category = Settings)
+	bool bParallelExecution = false;
+
+	/** Base name for newly-generated asset */
+	UPROPERTY(EditAnywhere, Category = Settings)
+	FGenerateStaticMeshLODProcessSettings GeneratorSettings;
+
 	UPROPERTY(VisibleAnywhere, Category = Previews)
 	TArray<UTexture2D*> PreviewTextures;
 
-	UPROPERTY(EditAnywhere, Category = General)
-	bool bParallelExecution = false;
-
 };
 
 
@@ -90,8 +95,10 @@ public:
 	virtual void SetWorld(UWorld* World);
 	virtual void SetAssetAPI(IAssetGenerationAPI* AssetAPI);
 
+	virtual void OnTick(float DeltaTime);
+
 	virtual bool HasCancel() const override { return true; }
-	virtual bool HasAccept() const override;
+	virtual bool HasAccept() const override { return true; }
 	virtual bool CanAccept() const override;
 
 protected:
@@ -106,17 +113,41 @@ protected:
 	TArray<UTexture2D*> PreviewTextures;
 
 	UPROPERTY()
-	TArray<UMaterialInterface*> PreviewMaterials;;
+	TArray<UMaterialInterface*> PreviewMaterials;
+
+
+protected:
+
+	UPROPERTY()
+	UCollisionGeometryVisualizationProperties* CollisionVizSettings = nullptr;
+
+	UPROPERTY()
+	UPhysicsObjectToolPropertySet* ObjectData = nullptr;
+
+	UPROPERTY()
+	UMaterialInterface* LineMaterial = nullptr;
+
+	UPROPERTY()
+	UPreviewGeometry* CollisionPreview;
 
 protected:
 	UWorld* TargetWorld;
 	IAssetGenerationAPI* AssetAPI;
 
-	TPimplPtr<FGenerateStaticMeshLODProcess> GenerateProcess;
+	TUniquePtr<FGenerateStaticMeshLODProcess> GenerateProcess;
+
+	void OnSettingsModified();
 
 	bool bPreviewValid;
 	void ValidatePreview();
 
+	bool bCollisionVisualizationDirty = false;
+	void UpdateCollisionVisualization();
+
+
 	void CreateNewAsset();
 	void UpdateExistingAsset();
+
+
+
 };
