MeshLODToolset: disable move of texture images because in some cases one texture is being copied into multiple Materials.
#rnx

diff --git a/Engine/Plugins/Experimental/MeshLODToolset/Source/MeshLODToolset/Private/Graphs/GenerateStaticMeshLODProcess.cpp b/Engine/Plugins/Experimental/MeshLODToolset/Source/MeshLODToolset/Private/Graphs/GenerateStaticMeshLODProcess.cpp
--- a/Engine/Plugins/Experimental/MeshLODToolset/Source/MeshLODToolset/Private/Graphs/GenerateStaticMeshLODProcess.cpp
+++ b/Engine/Plugins/Experimental/MeshLODToolset/Source/MeshLODToolset/Private/Graphs/GenerateStaticMeshLODProcess.cpp
@@ -82,6 +82,10 @@ bool FGenerateStaticMeshLODProcess::Initialize(UStaticMesh* StaticMeshIn)
 		SourceMaterials[mi].SourceMaterial = Materials[mi];
 
 		UMaterialInterface* MaterialInterface = Materials[mi].MaterialInterface;
+		if (MaterialInterface == nullptr)
+		{
+			continue;
+		}
 
 		TArray<FMaterialParameterInfo> ParameterInfo;
 		TArray<FGuid> ParameterIds;
@@ -314,7 +318,11 @@ bool FGenerateStaticMeshLODProcess::ComputeDerivedSourceData()
 				int32 BakedTexIndex = TextureToDerivedTexIndex[TexInfo.SourceTexture];
 				const TUniquePtr<UE::GeometryFlow::FTextureImage>& DerivedTex = DerivedTextureImages[BakedTexIndex];
 				TexInfo.Dimensions = DerivedTex->Image.GetDimensions();
-				TexInfo.Image = MoveTemp(DerivedTex->Image);
+
+				// Cannot currently MoveTemp here because this Texture may appear in multiple Materials, 
+				// and currently we do not handle that. The Materials need to learn how to share.
+				//TexInfo.Image = MoveTemp(DerivedTex->Image);
+				TexInfo.Image = DerivedTex->Image;
 			}
 		}
 	}
