MeshLODToolset: handle multi-use textures in output write
#rnx

diff --git a/Engine/Plugins/Experimental/MeshLODToolset/Source/MeshLODToolset/Private/Graphs/GenerateStaticMeshLODProcess.cpp b/Engine/Plugins/Experimental/MeshLODToolset/Source/MeshLODToolset/Private/Graphs/GenerateStaticMeshLODProcess.cpp
--- a/Engine/Plugins/Experimental/MeshLODToolset/Source/MeshLODToolset/Private/Graphs/GenerateStaticMeshLODProcess.cpp
+++ b/Engine/Plugins/Experimental/MeshLODToolset/Source/MeshLODToolset/Private/Graphs/GenerateStaticMeshLODProcess.cpp
@@ -460,6 +460,11 @@ void FGenerateStaticMeshLODProcess::WriteDerivedTextures()
 {
 	IAssetTools& AssetTools = FModuleManager::Get().LoadModuleChecked<FAssetToolsModule>("AssetTools").Get();
 
+	// this is a workaround for handling multiple materials that reference the same texture. Currently the code
+	// below will try to write that texture multiple times, which will fail when it tries to create a package
+	// for a filename that already exists
+	TArray<UTexture2D*> SourceTexturesWritten;
+
 	// write derived textures
 	int32 NumMaterials = SourceMaterials.Num();
 	check(DerivedMaterials.Num() == NumMaterials);
@@ -473,6 +478,11 @@ void FGenerateStaticMeshLODProcess::WriteDerivedTextures()
 		for (int32 ti = 0; ti < NumTextures; ++ti)
 		{
 			const FTextureInfo& SourceTex = SourceMaterialInfo.SourceTextures[ti];
+			if (SourceTexturesWritten.Contains(SourceTex.SourceTexture))
+			{
+				continue;
+			}
+
 			bool bConvertToSRGB = SourceTex.SourceTexture->SRGB;
 			FTextureInfo& DerivedTex = DerivedMaterialInfo.SourceTextures[ti];
 
@@ -495,6 +505,8 @@ void FGenerateStaticMeshLODProcess::WriteDerivedTextures()
 				// write asset
 				bool bWriteOK = WriteDerivedTexture(SourceTex.SourceTexture, DerivedTex.SourceTexture);
 				ensure(bWriteOK);
+
+				SourceTexturesWritten.Add(SourceTex.SourceTexture);
 			}
 		}
 
