AutoLOD: workaround for crash in paralllel graph evaluation.

diff --git a/Engine/Plugins/Experimental/MeshLODToolset/Source/MeshLODToolset/Private/Graphs/GenerateMeshLODGraph.cpp b/Engine/Plugins/Experimental/MeshLODToolset/Source/MeshLODToolset/Private/Graphs/GenerateMeshLODGraph.cpp
--- a/Engine/Plugins/Experimental/MeshLODToolset/Source/MeshLODToolset/Private/Graphs/GenerateMeshLODGraph.cpp
+++ b/Engine/Plugins/Experimental/MeshLODToolset/Source/MeshLODToolset/Private/Graphs/GenerateMeshLODGraph.cpp
@@ -46,19 +46,24 @@ void FGenerateMeshLODGraph::EvaluateResultParallel(
 	//FScopedDurationTimeLogger Timer(TEXT("FGenerateMeshLODGraph::EvaluateResult -- parallel execution"));
 
 	TArray<FGeometryFlowExecutor::NodeOutputSpec> DesiredOutputs;
+
 	DesiredOutputs.Add(FGeometryFlowExecutor::NodeOutputSpec{ BakeNormalMapNode, FBakeMeshNormalMapNode::OutParamNormalMap() });
+
 	for (FBakeTextureGraphInfo& TexBakeStep : BakeTextureNodes)
 	{
 		DesiredOutputs.Add(FGeometryFlowExecutor::NodeOutputSpec{
 			TexBakeStep.BakeNode,
 			FBakeMeshTextureImageNode::OutParamTextureImage() });
 	}
+
 	DesiredOutputs.Add(FGeometryFlowExecutor::NodeOutputSpec{
 			CollisionOutputNode,
 			FCollisionGeometryTransferNode::OutParamValue() });
+
 	DesiredOutputs.Add(FGeometryFlowExecutor::NodeOutputSpec{
 		TangentsOutputNode,
 		FMeshTangentsTransferNode::OutParamValue() });
+
 	DesiredOutputs.Add(FGeometryFlowExecutor::NodeOutputSpec{
 		MeshOutputNode,
 		FDynamicMeshTransferNode::OutParamValue() });
@@ -74,16 +79,18 @@ void FGenerateMeshLODGraph::EvaluateResultParallel(
 
 	// { BakeNormalMapNode, FBakeMeshNormalMapNode::OutParamNormalMap() }
 	NormalMap = FNormalMapImage();
-	bool bTakeNormalMap = true;
-	ExtractData(*OutputDataIter, NormalMap, (int)EMeshProcessingDataTypes::NormalMapImage, bTakeNormalMap);
+	bool bTakeNormalMap = false;
+	EGeometryFlowResult ExtractResult = ExtractData(*OutputDataIter, NormalMap, (int)EMeshProcessingDataTypes::NormalMapImage, bTakeNormalMap);
+	check(ExtractResult == EGeometryFlowResult::Ok);
 	++OutputDataIter;
 
 	for (FBakeTextureGraphInfo& TexBakeStep : BakeTextureNodes)
 	{
 		// { TexBakeStep.BakeNode, FBakeMeshTextureImageNode::OutParamTextureImage() }
 		TUniquePtr<UE::GeometryFlow::FTextureImage> NewImage = MakeUnique<UE::GeometryFlow::FTextureImage>();
-		bool bTakeTexBake = true;
-		ExtractData(*OutputDataIter, *NewImage, (int)EMeshProcessingDataTypes::TextureImage, bTakeTexBake);
+		bool bTakeTexBake = false;
+		ExtractResult = ExtractData(*OutputDataIter, *NewImage, (int)EMeshProcessingDataTypes::TextureImage, bTakeTexBake);
+		check(ExtractResult == EGeometryFlowResult::Ok);
 		++OutputDataIter;
 		TextureImages.Add(MoveTemp(NewImage));
 	}
@@ -91,20 +98,22 @@ void FGenerateMeshLODGraph::EvaluateResultParallel(
 	// {CollisionOutputNode, FCollisionGeometryTransferNode::OutParamValue() }
 	ResultCollision = FSimpleShapeSet3d();
 	bool bTakeResultCollision = false;
-	ExtractData(*OutputDataIter, ResultCollision, (int)FCollisionGeometry::DataTypeIdentifier, bTakeResultCollision);
+	ExtractResult = ExtractData(*OutputDataIter, ResultCollision, (int)FCollisionGeometry::DataTypeIdentifier, bTakeResultCollision);
+	check(ExtractResult == EGeometryFlowResult::Ok);
 	++OutputDataIter;
 
 	// { TangentsOutputNode, FMeshTangentsTransferNode::OutParamValue() }
 	bool bTakeResultTangents = false;
 	ResultTangents = FMeshTangentsd();
-	ExtractData(*OutputDataIter, ResultTangents, (int)EMeshProcessingDataTypes::MeshTangentSet, bTakeResultTangents);
+	ExtractResult = ExtractData(*OutputDataIter, ResultTangents, (int)EMeshProcessingDataTypes::MeshTangentSet, bTakeResultTangents);
+	check(ExtractResult == EGeometryFlowResult::Ok);
 	++OutputDataIter;
 
 	// {MeshOutputNode, FDynamicMeshTransferNode::OutParamValue() }
-	bool bTakeResultMesh = true;
+	bool bTakeResultMesh = false;
 	ResultMesh.Clear();
-	ExtractData(*OutputDataIter, ResultMesh, (int32)EMeshProcessingDataTypes::DynamicMesh, bTakeResultMesh);
-
+	ExtractResult = ExtractData(*OutputDataIter, ResultMesh, (int32)EMeshProcessingDataTypes::DynamicMesh, bTakeResultMesh);
+	check(ExtractResult == EGeometryFlowResult::Ok);
 	++OutputDataIter;
 }
 
@@ -300,7 +309,7 @@ void FGenerateMeshLODGraph::BuildGraph()
 
 	BakeCacheNode = Graph->AddNodeOfType<FMakeMeshBakingCacheNode>(TEXT("MakeBakeCache"));
 	Graph->AddConnection(MeshSourceNode, FDynamicMeshSourceNode::OutParamValue(), BakeCacheNode, FMakeMeshBakingCacheNode::InParamDetailMesh());
-	Graph->AddConnection(MeshOutputNode, FDynamicMeshTransferNode::OutParamValue(), BakeCacheNode, FMakeMeshBakingCacheNode::InParamTargetMesh());
+	Graph->AddConnection(RepackUVNode, FMeshRepackUVsNode::OutParamResultMesh(), BakeCacheNode, FMakeMeshBakingCacheNode::InParamTargetMesh());
 	BakeCacheSettingsNode = Graph->AddNodeOfType<FMeshMakeBakingCacheSettingsSourceNode>(TEXT("BakeCacheSettings"));
 	Graph->InferConnection(BakeCacheSettingsNode, BakeCacheNode);
 
