MeshLODToolset: convert FGenerateStaticMeshLODProcess to a UObject so that it can keep generated UTexture2D's alive, otherwise they may be GC'd while the write process is executing
#rnx

diff --git a/Engine/Plugins/Experimental/MeshLODToolset/Source/MeshLODToolset/Private/Graphs/GenerateStaticMeshLODProcess.cpp b/Engine/Plugins/Experimental/MeshLODToolset/Source/MeshLODToolset/Private/Graphs/GenerateStaticMeshLODProcess.cpp
--- a/Engine/Plugins/Experimental/MeshLODToolset/Source/MeshLODToolset/Private/Graphs/GenerateStaticMeshLODProcess.cpp
+++ b/Engine/Plugins/Experimental/MeshLODToolset/Source/MeshLODToolset/Private/Graphs/GenerateStaticMeshLODProcess.cpp
@@ -33,7 +33,7 @@
 #include "Engine/Classes/Components/StaticMeshComponent.h"
 #include "Engine/Classes/PhysicsEngine/BodySetup.h"
 
-#define LOCTEXT_NAMESPACE "FGenerateStaticMeshLODProcess"
+#define LOCTEXT_NAMESPACE "UGenerateStaticMeshLODProcess"
 
 
 namespace
@@ -54,7 +54,7 @@ struct FReadTextureJob
 };
 
 
-bool FGenerateStaticMeshLODProcess::Initialize(UStaticMesh* StaticMeshIn)
+bool UGenerateStaticMeshLODProcess::Initialize(UStaticMesh* StaticMeshIn)
 {
 	if (!ensure(StaticMeshIn)) return false;
 	if (!ensure(StaticMeshIn->GetNumSourceModels() > 0)) return false;
@@ -163,7 +163,7 @@ bool FGenerateStaticMeshLODProcess::Initialize(UStaticMesh* StaticMeshIn)
 
 
 
-void FGenerateStaticMeshLODProcess::CalculateDerivedPathName(FString NewAssetSuffix)
+void UGenerateStaticMeshLODProcess::CalculateDerivedPathName(FString NewAssetSuffix)
 {
 	DerivedSuffix = FPaths::MakeValidFileName(NewAssetSuffix);
 	if (DerivedSuffix.Len() == 0)
@@ -175,22 +175,22 @@ void FGenerateStaticMeshLODProcess::CalculateDerivedPathName(FString NewAssetSuf
 }
 
 
-bool FGenerateStaticMeshLODProcess::InitializeGenerator()
+bool UGenerateStaticMeshLODProcess::InitializeGenerator()
 {
 	Generator = MakeUnique<FGenerateMeshLODGraph>();
 	Generator->BuildGraph();
 
 	// initialize source textures
-	TextureToDerivedTexIndex.Reset();
+	SourceTextureToDerivedTexIndex.Reset();
 	for (const FMaterialInfo& MatInfo : SourceMaterials)
 	{
 		for (const FTextureInfo& TexInfo : MatInfo.SourceTextures)
 		{
 			if (TexInfo.bShouldBakeTexture &&
-				(TextureToDerivedTexIndex.Contains(TexInfo.SourceTexture) == false))
+				(SourceTextureToDerivedTexIndex.Contains(TexInfo.SourceTexture) == false))
 			{
 				int32 NewIndex = Generator->AppendTextureBakeNode(TexInfo.Image, TexInfo.SourceTexture->GetName());
-				TextureToDerivedTexIndex.Add(TexInfo.SourceTexture, NewIndex);
+				SourceTextureToDerivedTexIndex.Add(TexInfo.SourceTexture, NewIndex);
 			}
 		}
 	}
@@ -220,7 +220,7 @@ bool FGenerateStaticMeshLODProcess::InitializeGenerator()
 }
 
 
-void FGenerateStaticMeshLODProcess::UpdateSettings(const FGenerateStaticMeshLODProcessSettings& NewSettings)
+void UGenerateStaticMeshLODProcess::UpdateSettings(const FGenerateStaticMeshLODProcessSettings& NewSettings)
 {
 	bool bSharedVoxelResolutionChanged = (NewSettings.SolidifyVoxelResolution != CurrentSettings.SolidifyVoxelResolution);
 	if ( bSharedVoxelResolutionChanged
@@ -281,7 +281,7 @@ void FGenerateStaticMeshLODProcess::UpdateSettings(const FGenerateStaticMeshLODP
 
 
 
-bool FGenerateStaticMeshLODProcess::ComputeDerivedSourceData()
+bool UGenerateStaticMeshLODProcess::ComputeDerivedSourceData()
 {
 	DerivedTextureImages.Reset();
 	if (bUseParallelExecutor)
@@ -313,9 +313,9 @@ bool FGenerateStaticMeshLODProcess::ComputeDerivedSourceData()
 	{
 		for (FTextureInfo& TexInfo : MatInfo.SourceTextures)
 		{
-			if (TextureToDerivedTexIndex.Contains(TexInfo.SourceTexture))
+			if (SourceTextureToDerivedTexIndex.Contains(TexInfo.SourceTexture))
 			{
-				int32 BakedTexIndex = TextureToDerivedTexIndex[TexInfo.SourceTexture];
+				int32 BakedTexIndex = SourceTextureToDerivedTexIndex[TexInfo.SourceTexture];
 				const TUniquePtr<UE::GeometryFlow::FTextureImage>& DerivedTex = DerivedTextureImages[BakedTexIndex];
 				TexInfo.Dimensions = DerivedTex->Image.GetDimensions();
 
@@ -333,7 +333,7 @@ bool FGenerateStaticMeshLODProcess::ComputeDerivedSourceData()
 
 
 
-void FGenerateStaticMeshLODProcess::GetDerivedMaterialsPreview(FPreviewMaterials& MaterialSetOut)
+void UGenerateStaticMeshLODProcess::GetDerivedMaterialsPreview(FPreviewMaterials& MaterialSetOut)
 {
 	// force garbage collection of outstanding preview materials
 	CollectGarbage(GARBAGE_COLLECTION_KEEPFLAGS);
@@ -397,7 +397,7 @@ void FGenerateStaticMeshLODProcess::GetDerivedMaterialsPreview(FPreviewMaterials
 }
 
 
-void FGenerateStaticMeshLODProcess::UpdateMaterialTextureParameters(
+void UGenerateStaticMeshLODProcess::UpdateMaterialTextureParameters(
 	UMaterialInstanceDynamic* Material, 
 	const FMaterialInfo& SourceMaterialInfo,
 	const TMap<UTexture2D*, UTexture2D*>& PreviewTextures, 
@@ -434,29 +434,39 @@ void FGenerateStaticMeshLODProcess::UpdateMaterialTextureParameters(
 
 
 
-bool FGenerateStaticMeshLODProcess::WriteDerivedAssetData()
+bool UGenerateStaticMeshLODProcess::WriteDerivedAssetData()
 {
+	AllDerivedTextures.Reset();
+
 	WriteDerivedTextures();
 
 	WriteDerivedMaterials();
 
 	WriteDerivedStaticMeshAsset();
 
+	// clear list of derived textures we were holding onto to prevent GC
+	AllDerivedTextures.Reset();
+
 	return true;
 }
 
 
-void FGenerateStaticMeshLODProcess::UpdateSourceAsset()
+void UGenerateStaticMeshLODProcess::UpdateSourceAsset()
 {
+	AllDerivedTextures.Reset();
+
 	WriteDerivedTextures();
 
 	WriteDerivedMaterials();
 
 	UpdateSourceStaticMeshAsset();
+
+	// clear list of derived textures we were holding onto to prevent GC
+	AllDerivedTextures.Reset();
 }
 
 
-void FGenerateStaticMeshLODProcess::WriteDerivedTextures()
+void UGenerateStaticMeshLODProcess::WriteDerivedTextures()
 {
 	IAssetTools& AssetTools = FModuleManager::Get().LoadModuleChecked<FAssetToolsModule>("AssetTools").Get();
 
@@ -500,6 +510,8 @@ void FGenerateStaticMeshLODProcess::WriteDerivedTextures()
 			DerivedTex.SourceTexture = TextureBuilder.GetTexture2D();
 			if (ensure(DerivedTex.SourceTexture))
 			{
+				AllDerivedTextures.Add(DerivedTex.SourceTexture);
+
 				FTexture2DBuilder::CopyPlatformDataToSourceData(DerivedTex.SourceTexture, FTexture2DBuilder::ETextureType::Color);
 
 				// write asset
@@ -536,7 +548,7 @@ void FGenerateStaticMeshLODProcess::WriteDerivedTextures()
 
 
 
-bool FGenerateStaticMeshLODProcess::WriteDerivedTexture(UTexture2D* SourceTexture, UTexture2D* DerivedTexture)
+bool UGenerateStaticMeshLODProcess::WriteDerivedTexture(UTexture2D* SourceTexture, UTexture2D* DerivedTexture)
 {
 	IAssetTools& AssetTools = FModuleManager::Get().LoadModuleChecked<FAssetToolsModule>("AssetTools").Get();
 
@@ -546,7 +558,7 @@ bool FGenerateStaticMeshLODProcess::WriteDerivedTexture(UTexture2D* SourceTextur
 }
 
 
-bool FGenerateStaticMeshLODProcess::WriteDerivedTexture(UTexture2D* DerivedTexture, FString BaseTexName)
+bool UGenerateStaticMeshLODProcess::WriteDerivedTexture(UTexture2D* DerivedTexture, FString BaseTexName)
 {
 	IAssetTools& AssetTools = FModuleManager::Get().LoadModuleChecked<FAssetToolsModule>("AssetTools").Get();
 
@@ -584,8 +596,7 @@ bool FGenerateStaticMeshLODProcess::WriteDerivedTexture(UTexture2D* DerivedTextu
 }
 
 
-
-void FGenerateStaticMeshLODProcess::WriteDerivedMaterials()
+void UGenerateStaticMeshLODProcess::WriteDerivedMaterials()
 {
 	IAssetTools& AssetTools = FModuleManager::Get().LoadModuleChecked<FAssetToolsModule>("AssetTools").Get();
 
@@ -646,7 +657,7 @@ void FGenerateStaticMeshLODProcess::WriteDerivedMaterials()
 
 
 
-void FGenerateStaticMeshLODProcess::UpdateMaterialTextureParameters(UMaterialInstanceConstant* Material, FMaterialInfo& DerivedMaterialInfo)
+void UGenerateStaticMeshLODProcess::UpdateMaterialTextureParameters(UMaterialInstanceConstant* Material, FMaterialInfo& DerivedMaterialInfo)
 {
 	Material->Modify();
 
@@ -679,7 +690,7 @@ void FGenerateStaticMeshLODProcess::UpdateMaterialTextureParameters(UMaterialIns
 
 
 
-void FGenerateStaticMeshLODProcess::WriteDerivedStaticMeshAsset()
+void UGenerateStaticMeshLODProcess::WriteDerivedStaticMeshAsset()
 {
 	// [TODO] should we try to re-use existing asset here, or should we delete it? 
 	// The source asset might have had any number of config changes that we want to
@@ -750,7 +761,7 @@ void FGenerateStaticMeshLODProcess::WriteDerivedStaticMeshAsset()
 
 
 
-void FGenerateStaticMeshLODProcess::UpdateSourceStaticMeshAsset()
+void UGenerateStaticMeshLODProcess::UpdateSourceStaticMeshAsset()
 {
 	GEditor->BeginTransaction(LOCTEXT("UpdateExistingAssetMessage", "Added Generated LOD"));
 
diff --git a/Engine/Plugins/Experimental/MeshLODToolset/Source/MeshLODToolset/Private/Tools/GenerateStaticMeshLODAssetTool.cpp b/Engine/Plugins/Experimental/MeshLODToolset/Source/MeshLODToolset/Private/Tools/GenerateStaticMeshLODAssetTool.cpp
--- a/Engine/Plugins/Experimental/MeshLODToolset/Source/MeshLODToolset/Private/Tools/GenerateStaticMeshLODAssetTool.cpp
+++ b/Engine/Plugins/Experimental/MeshLODToolset/Source/MeshLODToolset/Private/Tools/GenerateStaticMeshLODAssetTool.cpp
@@ -113,7 +113,7 @@ void UGenerateStaticMeshLODAssetTool::Setup()
 		LOCTEXT("OnStartStaticMeshLODAssetTool", "This tool creates a new LOD asset"),
 		EToolMessageLevel::UserNotification);
 
-	GenerateProcess = MakeUnique<FGenerateStaticMeshLODProcess>();
+	GenerateProcess = NewObject<UGenerateStaticMeshLODProcess>(this);
 
 	TUniquePtr<FPrimitiveComponentTarget>& SourceComponent = ComponentTargets[0];
 	UStaticMeshComponent* StaticMeshComponent = CastChecked<UStaticMeshComponent>(SourceComponent->GetOwnerComponent());
@@ -160,7 +160,7 @@ void UGenerateStaticMeshLODAssetTool::Setup()
 	{
 		// TODO: We crash if we don't recreate the Process and reinitialize it. Why?
 
-		GenerateProcess = MakeUnique<FGenerateStaticMeshLODProcess>();
+		GenerateProcess = NewObject<UGenerateStaticMeshLODProcess>(this);
 
 		TUniquePtr<FPrimitiveComponentTarget>& SourceComponent = ComponentTargets[0];
 		UStaticMeshComponent* StaticMeshComponent = CastChecked<UStaticMeshComponent>(SourceComponent->GetOwnerComponent());
@@ -254,7 +254,7 @@ void UGenerateStaticMeshLODAssetTool::ValidatePreview()
 	});
 	PreviewMesh->UpdateTangents(&ResultTangents, true);
 
-	FGenerateStaticMeshLODProcess::FPreviewMaterials PreviewMaterialSet;
+	UGenerateStaticMeshLODProcess::FPreviewMaterials PreviewMaterialSet;
 	GenerateProcess->GetDerivedMaterialsPreview(PreviewMaterialSet);
 	if (PreviewMaterialSet.Materials.Num() > 0)
 	{
diff --git a/Engine/Plugins/Experimental/MeshLODToolset/Source/MeshLODToolset/Public/Graphs/GenerateStaticMeshLODProcess.h b/Engine/Plugins/Experimental/MeshLODToolset/Source/MeshLODToolset/Public/Graphs/GenerateStaticMeshLODProcess.h
--- a/Engine/Plugins/Experimental/MeshLODToolset/Source/MeshLODToolset/Public/Graphs/GenerateStaticMeshLODProcess.h
+++ b/Engine/Plugins/Experimental/MeshLODToolset/Source/MeshLODToolset/Public/Graphs/GenerateStaticMeshLODProcess.h
@@ -85,9 +85,11 @@ struct FGenerateStaticMeshLODProcessSettings
 };
 
 
-
-class FGenerateStaticMeshLODProcess
+UCLASS(Transient)
+class UGenerateStaticMeshLODProcess : public UObject
 {
+	GENERATED_BODY()
+
 public:
 
 	bool Initialize(UStaticMesh* SourceMesh);
@@ -135,7 +137,9 @@ public:
 
 protected:
 
+	UPROPERTY()
 	UStaticMesh* SourceStaticMesh;
+
 	FString SourceAssetPath;
 	FString SourceAssetFolder;
 	FString SourceAssetName;
@@ -172,9 +176,16 @@ protected:
 	FSimpleShapeSet3d DerivedCollision;
 	UE::GeometryFlow::FNormalMapImage DerivedNormalMapImage;
 	TArray<TUniquePtr<UE::GeometryFlow::FTextureImage>> DerivedTextureImages;
-	TMap<UTexture2D*, int32> TextureToDerivedTexIndex;
+	TMap<UTexture2D*, int32> SourceTextureToDerivedTexIndex;
 
 	TArray<FMaterialInfo> DerivedMaterials;
+
+	// This list is for accumulating derived UTexture2D's created during WriteDerivedTextures(). We have to
+	// maintain uproperty references to these or they may be garbage collected
+	UPROPERTY()
+	TSet<UTexture2D*> AllDerivedTextures;
+
+	UPROPERTY()
 	UTexture2D* DerivedNormalMapTex;
 
 	TUniquePtr<FGenerateMeshLODGraph> Generator;
diff --git a/Engine/Plugins/Experimental/MeshLODToolset/Source/MeshLODToolset/Public/Tools/GenerateStaticMeshLODAssetTool.h b/Engine/Plugins/Experimental/MeshLODToolset/Source/MeshLODToolset/Public/Tools/GenerateStaticMeshLODAssetTool.h
--- a/Engine/Plugins/Experimental/MeshLODToolset/Source/MeshLODToolset/Public/Tools/GenerateStaticMeshLODAssetTool.h
+++ b/Engine/Plugins/Experimental/MeshLODToolset/Source/MeshLODToolset/Public/Tools/GenerateStaticMeshLODAssetTool.h
@@ -134,7 +134,8 @@ protected:
 	UWorld* TargetWorld;
 	IAssetGenerationAPI* AssetAPI;
 
-	TUniquePtr<FGenerateStaticMeshLODProcess> GenerateProcess;
+	UPROPERTY()
+	UGenerateStaticMeshLODProcess* GenerateProcess;
 
 	void OnSettingsModified();
 
