GeometryFlowExecutor: Fix stealing outputs from nodes when using parallel execution. This change also effectively adds a second API to the parallel executor -- the user can now call a single synchronous function to get outputs, or they can call a function to kick off the multithreaded process, and then wait for individual outputs.

#rnx

diff --git a/Engine/Plugins/Experimental/GeometryFlow/Source/GeometryFlowCore/Private/GeometryFlowExecutor.cpp b/Engine/Plugins/Experimental/GeometryFlow/Source/GeometryFlowCore/Private/GeometryFlowExecutor.cpp
--- a/Engine/Plugins/Experimental/GeometryFlow/Source/GeometryFlowCore/Private/GeometryFlowExecutor.cpp
+++ b/Engine/Plugins/Experimental/GeometryFlow/Source/GeometryFlowCore/Private/GeometryFlowExecutor.cpp
@@ -28,43 +28,13 @@ namespace
 FGeometryFlowExecutor::FGeometryFlowExecutor(FGraph& InGraph)
 	: GeometryFlowGraph(InGraph)
 {
+	EvalInfo = MakeUnique<FEvaluationInfo>();
 	TopologicalSort();
-}
-
-
-void FGeometryFlowExecutor::ComputeOutputs(const TArray<NodeOutputSpec>& DesiredOutputs,
-										   TArray<TSafeSharedPtr<IData>>& OutputDatas)
-{
-	// TODO: Improve this!
-
-	CreateAndDispatchTaskGraph();
-
-	for (const NodeOutputSpec& OutputSpec : DesiredOutputs)
-	{
-		for (const TPair<FGraph::FHandle, FGraphEventRef>& NodeRef : GeometryFlowNodeToGraphTask)
-		{
-			FGraph::FHandle NodeHandle{ NodeRef.Key };
-			if ( OutputSpec.Key == NodeHandle)
-			{ 
-				// TODO: Don't wait for each one, see if we can set a callback for when a task is done
-				NodeRef.Value->Wait();
-
-				// TODO: Check desired data type matches node output type
-
-				TSafeSharedPtr<FNode> GeoFlowNode = GeometryFlowGraph.FindNode(NodeHandle);
-				TSafeSharedPtr<IData> OutputData = GeoFlowNode->GetOutput(OutputSpec.Value);
-				check(OutputData != nullptr);
-
-				OutputDatas.Add(OutputData);
-			}
-		}
-	}
 
+	// Check that TaskGraph is enabled!
+	check(FTaskGraphInterface::IsRunning());
 }
 
-
-
-
 void FGeometryFlowExecutor::TopologicalSort()
 {
 	// Kahn's algorithm from Wikipedia
@@ -111,7 +81,7 @@ void FGeometryFlowExecutor::TopologicalSort()
 }
 
 
-void FGeometryFlowExecutor::CreateAndDispatchTaskGraph()
+void FGeometryFlowExecutor::AsyncRunGraph()
 {
 	DebugNodeExecutionLog.Reset();
 	DebugNodeExecutionTime.Reset();
@@ -194,7 +164,6 @@ void FGeometryFlowExecutor::CreateAndDispatchTaskGraph()
 
 			double T0 = FPlatformTime::Seconds();
 			
-			TUniquePtr<FEvaluationInfo> EvalInfo = MakeUnique<FEvaluationInfo>();
 			Node->Evaluate(DatasIn, DatasOut, EvalInfo);
 
 			double T1 = FPlatformTime::Seconds();
@@ -211,7 +180,43 @@ void FGeometryFlowExecutor::CreateAndDispatchTaskGraph()
 
 		TPair<FGraph::FHandle, FGraphEventRef> NodeInfo{ NodeHandle, TaskGraphN };
 		GeometryFlowNodeToGraphTask.Add(NodeInfo);
-
 	}
 }
 
+
+EGeometryFlowResult FGeometryFlowExecutor::ComputeOutputs(const TArray<NodeOutputSpec>& DesiredOutputs,
+														  TArray<TSafeSharedPtr<IData>>& OutputDatas)
+{
+	AsyncRunGraph();
+
+	for (const NodeOutputSpec& Out : DesiredOutputs)
+	{
+		FGraphEventRef* GraphEvent = GeometryFlowNodeToGraphTask.Find(Out.NodeHandle);
+
+		if (GraphEvent == nullptr)
+		{
+			return EGeometryFlowResult::NoMatchesFound;
+		}
+
+		// Wait for the node to finish executing
+		(*GraphEvent)->Wait();
+
+		TSafeSharedPtr<FNode> GeoFlowNode = GeometryFlowGraph.FindNode(Out.NodeHandle);
+		if (!GeoFlowNode)
+		{
+			return EGeometryFlowResult::NodeDoesNotExist;
+		}
+
+		TSafeSharedPtr<IData> OutputData = (Out.bStealOutput) ?
+			GeoFlowNode->StealOutput(Out.OutputName) : GeoFlowNode->GetOutput(Out.OutputName);
+
+		if (!OutputData)
+		{
+			return EGeometryFlowResult::OutputDoesNotExist;
+		}
+
+		OutputDatas.Add(OutputData);
+	}
+
+	return EGeometryFlowResult::Ok;
+}
diff --git a/Engine/Plugins/Experimental/GeometryFlow/Source/GeometryFlowCore/Public/BaseNodes/BinaryOpLambdaNode.h b/Engine/Plugins/Experimental/GeometryFlow/Source/GeometryFlowCore/Public/BaseNodes/BinaryOpLambdaNode.h
--- a/Engine/Plugins/Experimental/GeometryFlow/Source/GeometryFlowCore/Public/BaseNodes/BinaryOpLambdaNode.h
+++ b/Engine/Plugins/Experimental/GeometryFlow/Source/GeometryFlowCore/Public/BaseNodes/BinaryOpLambdaNode.h
@@ -63,7 +63,7 @@ public:
 
 					TSafeSharedPtr<DataType> Result = MakeSafeShared<DataType>();
 					Result->SetData(BinaryOpResult);
-					SetOutput(ResultName(), Result);
+					SetOutput(OutParamResult(), Result);
 
 					EvaluationInfo->CountCompute(this);
 				}
diff --git a/Engine/Plugins/Experimental/GeometryFlow/Source/GeometryFlowCore/Public/GeometryFlowExecutor.h b/Engine/Plugins/Experimental/GeometryFlow/Source/GeometryFlowCore/Public/GeometryFlowExecutor.h
--- a/Engine/Plugins/Experimental/GeometryFlow/Source/GeometryFlowCore/Public/GeometryFlowExecutor.h
+++ b/Engine/Plugins/Experimental/GeometryFlow/Source/GeometryFlowCore/Public/GeometryFlowExecutor.h
@@ -20,8 +20,78 @@ public:
 
 	FGeometryFlowExecutor(FGraph& InGraph);
 
-	using NodeOutputSpec = TPair<FGraph::FHandle, FString>;		// (Node handle, output name)
-	void ComputeOutputs(const TArray<NodeOutputSpec>& DesiredOutputs, TArray<TSafeSharedPtr<IData>>& OutputDatas);
+	struct NodeOutputSpec
+	{
+		FGraph::FHandle NodeHandle;
+		FString OutputName;
+		bool bStealOutput;
+	};
+
+	// There are two ways to run the graph and get outputs:
+	// 1. Call ComputeOutputs, which is a synchronous call that will execute the graph and grab the outputs, or
+	// 2. First call AsyncRunGraph to launch the execution, then call GetOutput for each output.
+
+	// 1.
+	// Do graph execution, wait for it to finish, and get the requested output data.
+	// NOTE: If NodeOutputSpec::bStealOutput is false for a given output, you must *not* call IData::GiveTo on the 
+	// corresponding output data!
+	EGeometryFlowResult ComputeOutputs(const TArray<NodeOutputSpec>& DesiredOutputs,
+									   TArray<TSafeSharedPtr<IData>>& OutputDatas);
+
+	// 2a.
+	// Launch graph execution on background threads and return
+	void AsyncRunGraph();
+
+	// 2b.
+	// Wait for a given node to finish executing, then grab the requested data
+	template<typename T>
+	EGeometryFlowResult GetOutput(const FGraph::FHandle& NodeHandle,
+								  const FString& OutputName,
+								  T& Storage,
+								  int32 StorageTypeIdentifier,
+								  bool bStealOutputData)
+	{
+		FGraphEventRef* GraphEvent = GeometryFlowNodeToGraphTask.Find(NodeHandle);
+
+		if (GraphEvent == nullptr)
+		{
+			return EGeometryFlowResult::NoMatchesFound;
+		}
+
+		// Wait for the node to finish executing
+		(*GraphEvent)->Wait();
+
+		TSafeSharedPtr<FNode> GeoFlowNode = GeometryFlowGraph.FindNode(NodeHandle);
+		if (!GeoFlowNode)
+		{
+			return EGeometryFlowResult::NodeDoesNotExist;
+		}
+
+		TSafeSharedPtr<IData> OutputData = (bStealOutputData) ?
+			GeoFlowNode->StealOutput(OutputName) : GeoFlowNode->GetOutput(OutputName);
+
+		if (!OutputData)
+		{
+			return EGeometryFlowResult::OutputDoesNotExist;
+		}
+
+		if (OutputData->GetPayloadType() != StorageTypeIdentifier)
+		{
+			return EGeometryFlowResult::UnmatchedTypes;
+		}
+		if (bStealOutputData)
+		{
+			OutputData->GiveTo(Storage, StorageTypeIdentifier);
+		}
+		else
+		{
+			OutputData->GetDataCopy(Storage, StorageTypeIdentifier);
+		}
+
+		return EGeometryFlowResult::Ok;
+	}
+
+	TUniquePtr<FEvaluationInfo> EvalInfo;
 
 protected:
 
@@ -31,7 +101,6 @@ protected:
 	void TopologicalSort();
 
 	TMap<FGraph::FHandle, FGraphEventRef> GeometryFlowNodeToGraphTask;
-	void CreateAndDispatchTaskGraph();
 
 	// Debug helpers
 	TArray<int32> DebugNodeExecutionLog;		        // when a node executes it adds its ID to this array
diff --git a/Engine/Plugins/Experimental/MeshLODToolset/Source/MeshLODToolset/Private/Graphs/GenerateMeshLODGraph.cpp b/Engine/Plugins/Experimental/MeshLODToolset/Source/MeshLODToolset/Private/Graphs/GenerateMeshLODGraph.cpp
--- a/Engine/Plugins/Experimental/MeshLODToolset/Source/MeshLODToolset/Private/Graphs/GenerateMeshLODGraph.cpp
+++ b/Engine/Plugins/Experimental/MeshLODToolset/Source/MeshLODToolset/Private/Graphs/GenerateMeshLODGraph.cpp
@@ -36,6 +36,9 @@ void FGenerateMeshLODGraph::SetSourceMesh(const FDynamicMesh3& SourceMeshIn)
 	UpdateSourceNodeValue<FDynamicMeshSourceNode>(*Graph, MeshSourceNode, SourceMeshIn);
 }
 
+
+#if 0
+
 void FGenerateMeshLODGraph::EvaluateResultParallel(
 	FDynamicMesh3& ResultMesh,
 	FMeshTangentsd& ResultTangents,
@@ -47,40 +50,51 @@ void FGenerateMeshLODGraph::EvaluateResultParallel(
 
 	TArray<FGeometryFlowExecutor::NodeOutputSpec> DesiredOutputs;
 
-	DesiredOutputs.Add(FGeometryFlowExecutor::NodeOutputSpec{ BakeNormalMapNode, FBakeMeshNormalMapNode::OutParamNormalMap() });
+	bool bTakeNormalMap = true;
+	DesiredOutputs.Add(FGeometryFlowExecutor::NodeOutputSpec{ 
+		BakeNormalMapNode, 
+		FBakeMeshNormalMapNode::OutParamNormalMap(),
+		bTakeNormalMap });
 
+	bool bTakeTexBake = true;
 	for (FBakeTextureGraphInfo& TexBakeStep : BakeTextureNodes)
 	{
 		DesiredOutputs.Add(FGeometryFlowExecutor::NodeOutputSpec{
 			TexBakeStep.BakeNode,
-			FBakeMeshTextureImageNode::OutParamTextureImage() });
+			FBakeMeshTextureImageNode::OutParamTextureImage(),
+			bTakeTexBake });
 	}
 
+	bool bTakeResultCollision = false;
 	DesiredOutputs.Add(FGeometryFlowExecutor::NodeOutputSpec{
 			CollisionOutputNode,
-			FCollisionGeometryTransferNode::OutParamValue() });
+			FCollisionGeometryTransferNode::OutParamValue(),
+			bTakeResultCollision });
 
+	bool bTakeResultTangents = false;
 	DesiredOutputs.Add(FGeometryFlowExecutor::NodeOutputSpec{
 		TangentsOutputNode,
-		FMeshTangentsTransferNode::OutParamValue() });
+		FMeshTangentsTransferNode::OutParamValue(),
+		bTakeResultTangents });
 
+	bool bTakeResultMesh = true;
 	DesiredOutputs.Add(FGeometryFlowExecutor::NodeOutputSpec{
 		MeshOutputNode,
-		FDynamicMeshTransferNode::OutParamValue() });
+		FDynamicMeshTransferNode::OutParamValue(),
+		bTakeResultMesh });
 
 	FGeometryFlowExecutor Exec(*Graph);
-
 	TArray<TSafeSharedPtr<IData>> OutputDatas;
 	Exec.ComputeOutputs(DesiredOutputs, OutputDatas);
 
 	check(OutputDatas.Num() == DesiredOutputs.Num());
-
 	TArray<TSafeSharedPtr<IData>>::TIterator OutputDataIter = OutputDatas.CreateIterator();
 
-	// { BakeNormalMapNode, FBakeMeshNormalMapNode::OutParamNormalMap() }
 	NormalMap = FNormalMapImage();
-	bool bTakeNormalMap = false;
-	EGeometryFlowResult ExtractResult = ExtractData(*OutputDataIter, NormalMap, (int)EMeshProcessingDataTypes::NormalMapImage, bTakeNormalMap);
+	EGeometryFlowResult ExtractResult = ExtractData(*OutputDataIter,
+													NormalMap,
+													(int)EMeshProcessingDataTypes::NormalMapImage,
+													bTakeNormalMap);
 	check(ExtractResult == EGeometryFlowResult::Ok);
 	++OutputDataIter;
 
@@ -88,35 +102,98 @@ void FGenerateMeshLODGraph::EvaluateResultParallel(
 	{
 		// { TexBakeStep.BakeNode, FBakeMeshTextureImageNode::OutParamTextureImage() }
 		TUniquePtr<UE::GeometryFlow::FTextureImage> NewImage = MakeUnique<UE::GeometryFlow::FTextureImage>();
-		bool bTakeTexBake = false;
+		
 		ExtractResult = ExtractData(*OutputDataIter, *NewImage, (int)EMeshProcessingDataTypes::TextureImage, bTakeTexBake);
 		check(ExtractResult == EGeometryFlowResult::Ok);
-		++OutputDataIter;
 		TextureImages.Add(MoveTemp(NewImage));
+		++OutputDataIter;
 	}
 
 	// {CollisionOutputNode, FCollisionGeometryTransferNode::OutParamValue() }
 	ResultCollision = FSimpleShapeSet3d();
-	bool bTakeResultCollision = false;
 	ExtractResult = ExtractData(*OutputDataIter, ResultCollision, (int)FCollisionGeometry::DataTypeIdentifier, bTakeResultCollision);
 	check(ExtractResult == EGeometryFlowResult::Ok);
 	++OutputDataIter;
 
 	// { TangentsOutputNode, FMeshTangentsTransferNode::OutParamValue() }
-	bool bTakeResultTangents = false;
 	ResultTangents = FMeshTangentsd();
 	ExtractResult = ExtractData(*OutputDataIter, ResultTangents, (int)EMeshProcessingDataTypes::MeshTangentSet, bTakeResultTangents);
 	check(ExtractResult == EGeometryFlowResult::Ok);
 	++OutputDataIter;
 
 	// {MeshOutputNode, FDynamicMeshTransferNode::OutParamValue() }
-	bool bTakeResultMesh = false;
 	ResultMesh.Clear();
 	ExtractResult = ExtractData(*OutputDataIter, ResultMesh, (int32)EMeshProcessingDataTypes::DynamicMesh, bTakeResultMesh);
 	check(ExtractResult == EGeometryFlowResult::Ok);
 	++OutputDataIter;
 }
 
+#else
+
+void FGenerateMeshLODGraph::EvaluateResultParallel(
+	FDynamicMesh3& ResultMesh,
+	FMeshTangentsd& ResultTangents,
+	FSimpleShapeSet3d& ResultCollision,
+	UE::GeometryFlow::FNormalMapImage& NormalMap,
+	TArray<TUniquePtr<UE::GeometryFlow::FTextureImage>>& TextureImages)
+{
+	//FScopedDurationTimeLogger Timer(TEXT("FGenerateMeshLODGraph::EvaluateResult -- parallel execution"));
+
+	FGeometryFlowExecutor Exec(*Graph);
+	Exec.AsyncRunGraph();
+
+	NormalMap = FNormalMapImage();
+	bool bTakeNormalMap = true;
+	EGeometryFlowResult ExtractResult = Exec.GetOutput(BakeNormalMapNode,
+													   FBakeMeshNormalMapNode::OutParamNormalMap(),
+													   NormalMap, 
+													   (int)EMeshProcessingDataTypes::NormalMapImage, 
+													   bTakeNormalMap);
+	check(ExtractResult == EGeometryFlowResult::Ok);
+
+	for (FBakeTextureGraphInfo& TexBakeStep : BakeTextureNodes)
+	{
+		TUniquePtr<UE::GeometryFlow::FTextureImage> NewImage = MakeUnique<UE::GeometryFlow::FTextureImage>();
+		bool bTakeTexBake = true;
+		ExtractResult = Exec.GetOutput(TexBakeStep.BakeNode,
+									   FBakeMeshTextureImageNode::OutParamTextureImage(),
+									   *NewImage,
+									   (int)EMeshProcessingDataTypes::TextureImage,
+									   bTakeNormalMap);
+		check(ExtractResult == EGeometryFlowResult::Ok);
+		TextureImages.Add(MoveTemp(NewImage));
+	}
+
+	ResultCollision = FSimpleShapeSet3d();
+	bool bTakeResultCollision = false;
+	ExtractResult = Exec.GetOutput(CollisionOutputNode,
+								   FCollisionGeometryTransferNode::OutParamValue(),
+								   ResultCollision,
+								   (int)FCollisionGeometry::DataTypeIdentifier, bTakeResultCollision);
+	check(ExtractResult == EGeometryFlowResult::Ok);
+
+	bool bTakeResultTangents = false;
+	ResultTangents = FMeshTangentsd();
+	ExtractResult = Exec.GetOutput(TangentsOutputNode, 
+								   FMeshTangentsTransferNode::OutParamValue(),
+								   ResultTangents, 
+								   (int)EMeshProcessingDataTypes::MeshTangentSet, 
+								   bTakeResultTangents);
+	check(ExtractResult == EGeometryFlowResult::Ok);
+	
+
+	bool bTakeResultMesh = true;
+	ResultMesh.Clear();
+	ExtractResult = Exec.GetOutput(MeshOutputNode, 
+								   FDynamicMeshTransferNode::OutParamValue(),
+								   ResultMesh, 
+								   (int32)EMeshProcessingDataTypes::DynamicMesh, 
+								   bTakeResultMesh);
+	check(ExtractResult == EGeometryFlowResult::Ok);
+}
+
+#endif
+
 
 void FGenerateMeshLODGraph::UpdateSolidifySettings(const FMeshSolidifySettings& SolidifySettings)
 {
