Add SweptConvex and MinVolume options to the simple collision generator node.

#rnx

diff --git a/Engine/Plugins/Experimental/GeometryFlow/Source/GeometryFlowMeshProcessing/Private/PhysicsNodes/GenerateSimpleCollisionNode.cpp b/Engine/Plugins/Experimental/GeometryFlow/Source/GeometryFlowMeshProcessing/Private/PhysicsNodes/GenerateSimpleCollisionNode.cpp
--- a/Engine/Plugins/Experimental/GeometryFlow/Source/GeometryFlowMeshProcessing/Private/PhysicsNodes/GenerateSimpleCollisionNode.cpp
+++ b/Engine/Plugins/Experimental/GeometryFlow/Source/GeometryFlowMeshProcessing/Private/PhysicsNodes/GenerateSimpleCollisionNode.cpp
@@ -1,7 +1,6 @@
 // Copyright Epic Games, Inc. All Rights Reserved.
 
 #include "PhysicsNodes/GenerateSimpleCollisionNode.h"
-#include "ShapeApproximation/MeshSimpleShapeApproximation.h"
 #include "Operations/MeshConvexHull.h"
 #include "MeshIndexUtil.h"
 #include "Async/ParallelFor.h"
@@ -180,6 +179,15 @@ void FGenerateSimpleCollisionNode::Evaluate(
 				case ESimpleCollisionGeometryType::ConvexHulls:
 					GenerateSimpleCollisionNodeHelpers::GenerateConvexHulls(Mesh, IndexData, Settings, CollisionGeometry);
 					break;
+				case ESimpleCollisionGeometryType::SweptHulls:
+					ShapeApproximator.bDetectConvexes = true;					
+					ShapeApproximator.bSimplifyHulls = Settings.SweptHullSettings.bSimplifyPolygons;
+					ShapeApproximator.HullSimplifyTolerance = Settings.SweptHullSettings.HullTolerance;
+					ShapeApproximator.Generate_ProjectedHulls(CollisionGeometry.Geometry, Settings.SweptHullSettings.SweepAxis);
+					break;
+				case ESimpleCollisionGeometryType::MinVolume:
+					ShapeApproximator.Generate_MinVolume(CollisionGeometry.Geometry);
+					break;
 				}
 
 				SetOutput(OutParamGeometry(), MakeMovableData<FCollisionGeometry>(MoveTemp(CollisionGeometry)));
diff --git a/Engine/Plugins/Experimental/GeometryFlow/Source/GeometryFlowMeshProcessing/Public/PhysicsNodes/GenerateSimpleCollisionNode.h b/Engine/Plugins/Experimental/GeometryFlow/Source/GeometryFlowMeshProcessing/Public/PhysicsNodes/GenerateSimpleCollisionNode.h
--- a/Engine/Plugins/Experimental/GeometryFlow/Source/GeometryFlowMeshProcessing/Public/PhysicsNodes/GenerateSimpleCollisionNode.h
+++ b/Engine/Plugins/Experimental/GeometryFlow/Source/GeometryFlowMeshProcessing/Public/PhysicsNodes/GenerateSimpleCollisionNode.h
@@ -4,6 +4,7 @@
 
 #include "GeometryFlowCoreNodes.h"
 #include "GeometryFlowNodeUtil.h"
+#include "ShapeApproximation/MeshSimpleShapeApproximation.h"
 #include "MeshProcessingNodes/MeshProcessingDataTypes.h"
 #include "DataTypes/DynamicMeshData.h"
 #include "DataTypes/CollisionGeometryData.h"
@@ -17,11 +18,15 @@ namespace GeometryFlow
 
 enum class ESimpleCollisionGeometryType : uint8
 {
+	// NOTE: This must be kept in sync with EGenerateStaticMeshLODSimpleCollisionGeometryType in GenerateStaticMeshLODProcess.h
+
 	AlignedBoxes,
 	OrientedBoxes,
 	MinimalSpheres,
 	Capsules,
-	ConvexHulls
+	ConvexHulls,
+	SweptHulls,
+	MinVolume
 };
 
 
@@ -38,6 +43,12 @@ struct GEOMETRYFLOWMESHPROCESSING_API FGenerateSimpleCollisionSettings
 		int PrefilterGridResolution = 10;
 	} ConvexHullSettings;
 
+	struct FGenerateSweptHullSettings
+	{
+		bool bSimplifyPolygons = true;
+		FMeshSimpleShapeApproximation::EProjectedHullAxisMode SweepAxis = FMeshSimpleShapeApproximation::EProjectedHullAxisMode::SmallestVolume;
+		float HullTolerance = 0.1;
+	} SweptHullSettings;
 };
 
 GEOMETRYFLOW_DECLARE_SETTINGS_TYPES(FGenerateSimpleCollisionSettings, GenerateSimpleCollision);
diff --git a/Engine/Plugins/Experimental/MeshLODToolset/Source/MeshLODToolset/Private/Graphs/GenerateStaticMeshLODProcess.cpp b/Engine/Plugins/Experimental/MeshLODToolset/Source/MeshLODToolset/Private/Graphs/GenerateStaticMeshLODProcess.cpp
--- a/Engine/Plugins/Experimental/MeshLODToolset/Source/MeshLODToolset/Private/Graphs/GenerateStaticMeshLODProcess.cpp
+++ b/Engine/Plugins/Experimental/MeshLODToolset/Source/MeshLODToolset/Private/Graphs/GenerateStaticMeshLODProcess.cpp
@@ -213,10 +213,17 @@ bool UGenerateStaticMeshLODProcess::InitializeGenerator()
 	CurrentSettings.BakeResolution = (EGenerateStaticMeshLODBakeResolution)Generator->GetCurrentBakeCacheSettings().Dimensions.GetWidth();
 	CurrentSettings.BakeThickness = Generator->GetCurrentBakeCacheSettings().Thickness;
 
-	CurrentSettings.CollisionType = static_cast<EGenerateStaticMeshSimpleCollisionGeometryType>(Generator->GetCurrentGenerateSimpleCollisionSettings().Type);
-	CurrentSettings.ConvexTriangleCount = Generator->GetCurrentGenerateSimpleCollisionSettings().ConvexHullSettings.SimplifyToTriangleCount;
-	CurrentSettings.bPrefilterVertices = Generator->GetCurrentGenerateSimpleCollisionSettings().ConvexHullSettings.bPrefilterVertices;
-	CurrentSettings.PrefilterGridResolution = Generator->GetCurrentGenerateSimpleCollisionSettings().ConvexHullSettings.PrefilterGridResolution;
+	
+	const UE::GeometryFlow::FGenerateSimpleCollisionSettings& SimpleCollisionSettings = Generator->GetCurrentGenerateSimpleCollisionSettings();
+	CurrentSettings.CollisionType = static_cast<EGenerateStaticMeshLODSimpleCollisionGeometryType>(SimpleCollisionSettings.Type);
+	CurrentSettings.ConvexTriangleCount = SimpleCollisionSettings.ConvexHullSettings.SimplifyToTriangleCount;
+	CurrentSettings.bPrefilterVertices = SimpleCollisionSettings.ConvexHullSettings.bPrefilterVertices;
+	CurrentSettings.PrefilterGridResolution = SimpleCollisionSettings.ConvexHullSettings.PrefilterGridResolution;
+	CurrentSettings.bSimplifyPolygons = SimpleCollisionSettings.SweptHullSettings.bSimplifyPolygons;
+	CurrentSettings.HullTolerance = SimpleCollisionSettings.SweptHullSettings.HullTolerance;
+
+	FMeshSimpleShapeApproximation::EProjectedHullAxisMode RHSMode = SimpleCollisionSettings.SweptHullSettings.SweepAxis;
+	CurrentSettings.SweepAxis = static_cast<EGenerateStaticMeshLODProjectedHullAxisMode>(RHSMode);
 
 
 	return true;
@@ -273,6 +280,9 @@ void UGenerateStaticMeshLODProcess::UpdateSettings(const FGenerateStaticMeshLODP
 	if (NewSettings.ConvexTriangleCount != CurrentSettings.ConvexTriangleCount ||
 		NewSettings.bPrefilterVertices != CurrentSettings.bPrefilterVertices ||
 		NewSettings.PrefilterGridResolution != CurrentSettings.PrefilterGridResolution ||
+		NewSettings.bSimplifyPolygons != CurrentSettings.bSimplifyPolygons ||
+		NewSettings.HullTolerance != CurrentSettings.HullTolerance ||
+		NewSettings.SweepAxis != CurrentSettings.SweepAxis ||
 		NewSettings.CollisionType != CurrentSettings.CollisionType)
 	{
 		UE::GeometryFlow::FGenerateSimpleCollisionSettings NewGenCollisionSettings = Generator->GetCurrentGenerateSimpleCollisionSettings();
@@ -280,6 +290,9 @@ void UGenerateStaticMeshLODProcess::UpdateSettings(const FGenerateStaticMeshLODP
 		NewGenCollisionSettings.ConvexHullSettings.SimplifyToTriangleCount = NewSettings.ConvexTriangleCount;
 		NewGenCollisionSettings.ConvexHullSettings.bPrefilterVertices = NewSettings.bPrefilterVertices;
 		NewGenCollisionSettings.ConvexHullSettings.PrefilterGridResolution = NewSettings.PrefilterGridResolution;
+		NewGenCollisionSettings.SweptHullSettings.bSimplifyPolygons = NewSettings.bSimplifyPolygons;
+		NewGenCollisionSettings.SweptHullSettings.HullTolerance = NewSettings.HullTolerance;
+		NewGenCollisionSettings.SweptHullSettings.SweepAxis = static_cast<FMeshSimpleShapeApproximation::EProjectedHullAxisMode>(NewSettings.SweepAxis);
 		Generator->UpdateGenerateSimpleCollisionSettings(NewGenCollisionSettings);
 	}
 
diff --git a/Engine/Plugins/Experimental/MeshLODToolset/Source/MeshLODToolset/Private/Tools/GenerateStaticMeshLODAssetTool.cpp b/Engine/Plugins/Experimental/MeshLODToolset/Source/MeshLODToolset/Private/Tools/GenerateStaticMeshLODAssetTool.cpp
--- a/Engine/Plugins/Experimental/MeshLODToolset/Source/MeshLODToolset/Private/Tools/GenerateStaticMeshLODAssetTool.cpp
+++ b/Engine/Plugins/Experimental/MeshLODToolset/Source/MeshLODToolset/Private/Tools/GenerateStaticMeshLODAssetTool.cpp
@@ -132,11 +132,13 @@ void UGenerateStaticMeshLODAssetTool::Setup()
 	BasicProperties->WatchProperty(BasicProperties->GeneratorSettings.NumAutoUVCharts, [this](int) { OnSettingsModified(); });
 	BasicProperties->WatchProperty(BasicProperties->GeneratorSettings.BakeResolution, [this](EGenerateStaticMeshLODBakeResolution) { OnSettingsModified(); });
 	BasicProperties->WatchProperty(BasicProperties->GeneratorSettings.BakeThickness, [this](float) { OnSettingsModified(); });
-	BasicProperties->WatchProperty(BasicProperties->GeneratorSettings.CollisionType, [this](EGenerateStaticMeshSimpleCollisionGeometryType) { OnSettingsModified(); });
+	BasicProperties->WatchProperty(BasicProperties->GeneratorSettings.CollisionType, [this](EGenerateStaticMeshLODSimpleCollisionGeometryType) { OnSettingsModified(); });
 	BasicProperties->WatchProperty(BasicProperties->GeneratorSettings.ConvexTriangleCount, [this](int) { OnSettingsModified(); });
 	BasicProperties->WatchProperty(BasicProperties->GeneratorSettings.bPrefilterVertices, [this](bool) { OnSettingsModified(); });
 	BasicProperties->WatchProperty(BasicProperties->GeneratorSettings.PrefilterGridResolution, [this](float) { OnSettingsModified(); });
-
+	BasicProperties->WatchProperty(BasicProperties->GeneratorSettings.bSimplifyPolygons, [this](bool) { OnSettingsModified(); });
+	BasicProperties->WatchProperty(BasicProperties->GeneratorSettings.HullTolerance, [this](float) { OnSettingsModified(); });
+	BasicProperties->WatchProperty(BasicProperties->GeneratorSettings.SweepAxis, [this](EGenerateStaticMeshLODProjectedHullAxisMode) { OnSettingsModified(); });
 
 	FBoxSphereBounds Bounds = StaticMeshComponent->Bounds;
 	FTransform PreviewTransform = SourceComponent->GetWorldTransform();
diff --git a/Engine/Plugins/Experimental/MeshLODToolset/Source/MeshLODToolset/Public/Graphs/GenerateStaticMeshLODProcess.h b/Engine/Plugins/Experimental/MeshLODToolset/Source/MeshLODToolset/Public/Graphs/GenerateStaticMeshLODProcess.h
--- a/Engine/Plugins/Experimental/MeshLODToolset/Source/MeshLODToolset/Public/Graphs/GenerateStaticMeshLODProcess.h
+++ b/Engine/Plugins/Experimental/MeshLODToolset/Source/MeshLODToolset/Public/Graphs/GenerateStaticMeshLODProcess.h
@@ -33,13 +33,29 @@ enum class EGenerateStaticMeshLODBakeResolution
 };
 
 UENUM()
-enum class EGenerateStaticMeshSimpleCollisionGeometryType : uint8
+enum class EGenerateStaticMeshLODSimpleCollisionGeometryType : uint8
 {
+	// NOTE: This must be kept in sync with ESimpleCollisionGeometryType in GenerateSimpleCollisionNode.h
+
 	AlignedBoxes,
 	OrientedBoxes,
 	MinimalSpheres,
 	Capsules,
-	ConvexHulls
+	ConvexHulls,
+	SweptHulls,
+	MinVolume
+};
+
+UENUM()
+enum class EGenerateStaticMeshLODProjectedHullAxisMode : uint8
+{
+	// NOTE: This must be kept in sync with FMeshSimpleShapeApproximation::EProjectedHullAxisMode in MeshSimpleShapeApproximation.h
+
+	X = 0,
+	Y = 1,
+	Z = 2,
+	SmallestBoxDimension = 3,
+	SmallestVolume = 4
 };
 
 
@@ -91,19 +107,32 @@ struct FGenerateStaticMeshLODProcessSettings
 	// Convex Hull Settings
 
 	UPROPERTY(EditAnywhere, Category = Collision, meta = (DisplayName = "Collision Type"))
-	EGenerateStaticMeshSimpleCollisionGeometryType CollisionType = EGenerateStaticMeshSimpleCollisionGeometryType::ConvexHulls;
+	EGenerateStaticMeshLODSimpleCollisionGeometryType CollisionType = EGenerateStaticMeshLODSimpleCollisionGeometryType::ConvexHulls;
 
 	UPROPERTY(EditAnywhere, Category = Collision, meta = (DisplayName = "Convex Tri Count", 
-														  EditCondition = "CollisionType == EGenerateStaticMeshSimpleCollisionGeometryType::ConvexHulls"))
+														  EditCondition = "CollisionType == EGenerateStaticMeshLODSimpleCollisionGeometryType::ConvexHulls"))
 	int ConvexTriangleCount = 50;
 
-	UPROPERTY(EditAnywhere, Category = Collision, meta = (EditCondition = "CollisionType == EGenerateStaticMeshSimpleCollisionGeometryType::ConvexHulls"))
+	UPROPERTY(EditAnywhere, Category = Collision, meta = (EditCondition = "CollisionType == EGenerateStaticMeshLODSimpleCollisionGeometryType::ConvexHulls"))
 	bool bPrefilterVertices = false;
 
 	/** Grid resolution (along the maximum-length axis) */
-	UPROPERTY(EditAnywhere, Category = Collision, meta = (EditCondition = "CollisionType == EGenerateStaticMeshSimpleCollisionGeometryType::ConvexHulls && bPrefilterVertices", UIMin = 1, UIMax = 30))
+	UPROPERTY(EditAnywhere, Category = Collision, meta = (EditCondition = "CollisionType == EGenerateStaticMeshLODSimpleCollisionGeometryType::ConvexHulls && bPrefilterVertices", UIMin = 1, UIMax = 30))
 	int PrefilterGridResolution = 10;
 
+
+	// Swept Convex Hull Settings
+
+	UPROPERTY(EditAnywhere, Category = Collision, meta = (EditCondition = "CollisionType == EGenerateStaticMeshLODSimpleCollisionGeometryType::SweptHulls"))
+	bool bSimplifyPolygons = true;
+
+	UPROPERTY(EditAnywhere, Category = Collision, meta = (UIMin = "0", UIMax = "10", ClampMin = "0", ClampMax = "100000",
+														  EditCondition = "CollisionType == EGenerateStaticMeshLODSimpleCollisionGeometryType::SweptHulls"))
+	float HullTolerance = 0.1;
+
+	UPROPERTY(EditAnywhere, Category = Collision, meta = (EditCondition = "CollisionType == EGenerateStaticMeshLODSimpleCollisionGeometryType::SweptHulls"))
+	EGenerateStaticMeshLODProjectedHullAxisMode SweepAxis = EGenerateStaticMeshLODProjectedHullAxisMode::SmallestVolume;
+
 };
 
 
