Wire up "Default" name to group layer name input on mesh segmentation node.

#rnx

diff --git a/Engine/Plugins/Experimental/MeshLODToolset/Source/MeshLODToolset/Private/Graphs/GenerateMeshLODGraph.cpp b/Engine/Plugins/Experimental/MeshLODToolset/Source/MeshLODToolset/Private/Graphs/GenerateMeshLODGraph.cpp
--- a/Engine/Plugins/Experimental/MeshLODToolset/Source/MeshLODToolset/Private/Graphs/GenerateMeshLODGraph.cpp
+++ b/Engine/Plugins/Experimental/MeshLODToolset/Source/MeshLODToolset/Private/Graphs/GenerateMeshLODGraph.cpp
@@ -402,11 +402,13 @@ void FGenerateMeshLODGraph::BuildGraph()
 	// collision generation
 
 	FGraph::FHandle IgnoreGroupsForCollisionNode = Graph->AddNodeOfType<FIndexSetsSourceNode>(TEXT("CollisionIgnoreGroups"));
+	FGraph::FHandle GroupLayerNameNode = Graph->AddNodeOfType<FNameSourceNode>(TEXT("GroupLayerNameSource"));
 
 	//DecomposeMeshForCollisionNode = Graph->AddNodeOfType<FMakeTriangleSetsFromMeshNode>(TEXT("Decompose"));
 	DecomposeMeshForCollisionNode = Graph->AddNodeOfType<FMakeTriangleSetsFromGroupsNode>(TEXT("Decompose"));
 	Graph->InferConnection(MeshSourceNode, DecomposeMeshForCollisionNode);
 	Graph->InferConnection(IgnoreGroupsForCollisionNode, DecomposeMeshForCollisionNode);
+	Graph->InferConnection(GroupLayerNameNode, DecomposeMeshForCollisionNode);
 
 	GenerateSimpleCollisionNode = Graph->AddNodeOfType<FGenerateSimpleCollisionNode>(TEXT("GenerateSimpleCollision"));
 	Graph->InferConnection(MeshSourceNode, GenerateSimpleCollisionNode);
@@ -473,6 +475,9 @@ void FGenerateMeshLODGraph::BuildGraph()
 	IgnoreGroupsForCollision.AppendSet({ 0 });
 	UpdateSettingsSourceNodeValue(*Graph, IgnoreGroupsForCollisionNode, IgnoreGroupsForCollision);
 
+	FName GroupLayerName{ "Default" };
+	UpdateSourceNodeValue<FNameSourceNode>(*Graph, GroupLayerNameNode, GroupLayerName);
+
 	FGenerateSimpleCollisionSettings GenSimpleCollisionSettings;
 	UpdateGenerateSimpleCollisionSettings(GenSimpleCollisionSettings);
 
