AutoLOD: Add PreFilter Triangles Pass to GenerateMeshLODGraph, deletes triangles with GroupID != 0, from a named group layer. Expose group layer in Process/Tool settings.
#rnx

diff --git a/Engine/Plugins/Experimental/MeshLODToolset/Source/MeshLODToolset/Private/Graphs/GenerateMeshLODGraph.cpp b/Engine/Plugins/Experimental/MeshLODToolset/Source/MeshLODToolset/Private/Graphs/GenerateMeshLODGraph.cpp
--- a/Engine/Plugins/Experimental/MeshLODToolset/Source/MeshLODToolset/Private/Graphs/GenerateMeshLODGraph.cpp
+++ b/Engine/Plugins/Experimental/MeshLODToolset/Source/MeshLODToolset/Private/Graphs/GenerateMeshLODGraph.cpp
@@ -195,6 +195,13 @@ void FGenerateMeshLODGraph::EvaluateResultParallel(
 #endif
 
 
+void FGenerateMeshLODGraph::UpdatePreFilterSettings(const FMeshLODGraphPreFilterSettings& PreFilterSettings)
+{
+	UpdateSourceNodeValue<FNameSourceNode>(*Graph, FilterGroupsLayerNameNode, PreFilterSettings.FilterGroupLayerName);
+	CurrentPreFilterSettings = PreFilterSettings;
+}
+
+
 void FGenerateMeshLODGraph::UpdateSolidifySettings(const FMeshSolidifySettings& SolidifySettings)
 {
 	UpdateSettingsSourceNodeValue(*Graph, SolidifySettingsNode, SolidifySettings);
@@ -325,10 +332,24 @@ void FGenerateMeshLODGraph::BuildGraph()
 
 	MeshSourceNode = Graph->AddNodeOfType<FDynamicMeshSourceNode>(TEXT("SourceMesh"));
 
+	// remove detail triangles
+
+	FGraph::FHandle FilterGroupsNode = Graph->AddNodeOfType<FIndexSetsSourceNode>(TEXT("FilterGroups"));
+	FilterGroupsLayerNameNode = Graph->AddNodeOfType<FNameSourceNode>(TEXT("FilterGroupsLayerNameSource"));
+
+	FGraph::FHandle MakeFilterTriangleSetsNode = Graph->AddNodeOfType<FMakeTriangleSetsFromGroupsNode>(TEXT("MakeFilterTriangles"));
+	Graph->InferConnection(MeshSourceNode, MakeFilterTriangleSetsNode);
+	Graph->InferConnection(FilterGroupsNode, MakeFilterTriangleSetsNode);
+	Graph->InferConnection(FilterGroupsLayerNameNode, MakeFilterTriangleSetsNode);
+
+	FilterTrianglesNode = Graph->AddNodeOfType<FMeshDeleteTrianglesNode>(TEXT("FilterMesh"));
+	Graph->InferConnection(MeshSourceNode, FilterTrianglesNode);
+	Graph->InferConnection(MakeFilterTriangleSetsNode, FilterTrianglesNode);
+
 	// generating low-poly mesh
 
 	SolidifyNode = Graph->AddNodeOfType<FSolidifyMeshNode>(TEXT("Solidify"));
-	Graph->InferConnection(MeshSourceNode, SolidifyNode);
+	Graph->InferConnection(FilterTrianglesNode, SolidifyNode);
 	SolidifySettingsNode = Graph->AddNodeOfType<FSolidifySettingsSourceNode>(TEXT("SolidifySettings"));
 	Graph->InferConnection(SolidifySettingsNode, SolidifyNode);
 
@@ -405,11 +426,11 @@ void FGenerateMeshLODGraph::BuildGraph()
 
 	//DecomposeMeshForCollisionNode = Graph->AddNodeOfType<FMakeTriangleSetsFromMeshNode>(TEXT("Decompose"));
 	DecomposeMeshForCollisionNode = Graph->AddNodeOfType<FMakeTriangleSetsFromGroupsNode>(TEXT("Decompose"));
-	Graph->InferConnection(MeshSourceNode, DecomposeMeshForCollisionNode);
+	Graph->InferConnection(FilterTrianglesNode, DecomposeMeshForCollisionNode);
 	Graph->InferConnection(IgnoreGroupsForCollisionNode, DecomposeMeshForCollisionNode);
 
 	GenerateSimpleCollisionNode = Graph->AddNodeOfType<FGenerateSimpleCollisionNode>(TEXT("GenerateSimpleCollision"));
-	Graph->InferConnection(MeshSourceNode, GenerateSimpleCollisionNode);
+	Graph->InferConnection(FilterTrianglesNode, GenerateSimpleCollisionNode);
 	Graph->InferConnection(DecomposeMeshForCollisionNode, GenerateSimpleCollisionNode);
 	GenerateSimpleCollisionSettingsNode = Graph->AddNodeOfType<FGenerateSimpleCollisionSettingsSourceNode>(TEXT("GenerateSimpleCollisionSettings"));
 	Graph->InferConnection(GenerateSimpleCollisionSettingsNode, GenerateSimpleCollisionNode);
@@ -425,6 +446,14 @@ void FGenerateMeshLODGraph::BuildGraph()
 	//
 
 
+	FIndexSets IgnoreGroupsForDelete;
+	IgnoreGroupsForDelete.AppendSet({ 0 });
+	UpdateSettingsSourceNodeValue(*Graph, FilterGroupsNode, IgnoreGroupsForDelete);
+
+	FMeshLODGraphPreFilterSettings PreFilterSettings;
+	PreFilterSettings.FilterGroupLayerName = FName(TEXT("PreFilterGroups"));
+	UpdatePreFilterSettings(PreFilterSettings);
+
 	FMeshSolidifySettings SolidifySettings;
 	UpdateSolidifySettings(SolidifySettings);
 
diff --git a/Engine/Plugins/Experimental/MeshLODToolset/Source/MeshLODToolset/Private/Graphs/GenerateStaticMeshLODProcess.cpp b/Engine/Plugins/Experimental/MeshLODToolset/Source/MeshLODToolset/Private/Graphs/GenerateStaticMeshLODProcess.cpp
--- a/Engine/Plugins/Experimental/MeshLODToolset/Source/MeshLODToolset/Private/Graphs/GenerateStaticMeshLODProcess.cpp
+++ b/Engine/Plugins/Experimental/MeshLODToolset/Source/MeshLODToolset/Private/Graphs/GenerateStaticMeshLODProcess.cpp
@@ -202,6 +202,8 @@ bool UGenerateStaticMeshLODProcess::InitializeGenerator()
 
 	// read back default settings
 
+	CurrentSettings.FilterGroupLayer = Generator->GetCurrentPreFilterSettings().FilterGroupLayerName;
+
 	CurrentSettings.SolidifyVoxelResolution = Generator->GetCurrentSolidifySettings().VoxelResolution;
 	CurrentSettings.WindingThreshold = Generator->GetCurrentSolidifySettings().WindingThreshold;
 
@@ -230,73 +232,81 @@ bool UGenerateStaticMeshLODProcess::InitializeGenerator()
 }
 
 
-void UGenerateStaticMeshLODProcess::UpdateSettings(const FGenerateStaticMeshLODProcessSettings& NewSettings)
+void UGenerateStaticMeshLODProcess::UpdateSettings(const FGenerateStaticMeshLODProcessSettings& NewCombinedSettings)
 {
-	bool bSharedVoxelResolutionChanged = (NewSettings.SolidifyVoxelResolution != CurrentSettings.SolidifyVoxelResolution);
+
+	if (NewCombinedSettings.FilterGroupLayer != CurrentSettings.FilterGroupLayer)
+	{
+		FMeshLODGraphPreFilterSettings NewPreFilterSettings = Generator->GetCurrentPreFilterSettings();
+		NewPreFilterSettings.FilterGroupLayerName = NewCombinedSettings.FilterGroupLayer;
+		Generator->UpdatePreFilterSettings(NewPreFilterSettings);
+	}
+
+	bool bSharedVoxelResolutionChanged = (NewCombinedSettings.SolidifyVoxelResolution != CurrentSettings.SolidifyVoxelResolution);
 	if ( bSharedVoxelResolutionChanged
-		|| (NewSettings.WindingThreshold != CurrentSettings.WindingThreshold))
+		|| (NewCombinedSettings.WindingThreshold != CurrentSettings.WindingThreshold))
 	{
 		UE::GeometryFlow::FMeshSolidifySettings NewSolidifySettings = Generator->GetCurrentSolidifySettings();
-		NewSolidifySettings.VoxelResolution = NewSettings.SolidifyVoxelResolution;
-		NewSolidifySettings.WindingThreshold = NewSettings.WindingThreshold;
+		NewSolidifySettings.VoxelResolution = NewCombinedSettings.SolidifyVoxelResolution;
+		NewSolidifySettings.WindingThreshold = NewCombinedSettings.WindingThreshold;
 		Generator->UpdateSolidifySettings(NewSolidifySettings);
 	}
 
 
 	if ( bSharedVoxelResolutionChanged
-		|| (NewSettings.ClosureDistance != CurrentSettings.ClosureDistance))
+		|| (NewCombinedSettings.ClosureDistance != CurrentSettings.ClosureDistance))
 	{
 		UE::GeometryFlow::FVoxClosureSettings NewClosureSettings = Generator->GetCurrentMorphologySettings();
-		NewClosureSettings.VoxelResolution = NewSettings.SolidifyVoxelResolution;
-		NewClosureSettings.Distance = NewSettings.ClosureDistance;
+		NewClosureSettings.VoxelResolution = NewCombinedSettings.SolidifyVoxelResolution;
+		NewClosureSettings.Distance = NewCombinedSettings.ClosureDistance;
 		Generator->UpdateMorphologySettings(NewClosureSettings);
 	}
 
 
-	if (NewSettings.SimplifyTriangleCount != CurrentSettings.SimplifyTriangleCount)
+	if (NewCombinedSettings.SimplifyTriangleCount != CurrentSettings.SimplifyTriangleCount)
 	{
 		UE::GeometryFlow::FMeshSimplifySettings NewSimplifySettings = Generator->GetCurrentSimplifySettings();
-		NewSimplifySettings.TargetCount = NewSettings.SimplifyTriangleCount;
+		NewSimplifySettings.TargetCount = NewCombinedSettings.SimplifyTriangleCount;
 		Generator->UpdateSimplifySettings(NewSimplifySettings);
 	}
 
-	if (NewSettings.NumAutoUVCharts != CurrentSettings.NumAutoUVCharts)
+	if (NewCombinedSettings.NumAutoUVCharts != CurrentSettings.NumAutoUVCharts)
 	{
 		UE::GeometryFlow::FMeshAutoGenerateUVsSettings NewAutoUVSettings = Generator->GetCurrentAutoUVSettings();
-		NewAutoUVSettings.NumCharts = NewSettings.NumAutoUVCharts;
+		NewAutoUVSettings.NumCharts = NewCombinedSettings.NumAutoUVCharts;
 		Generator->UpdateAutoUVSettings(NewAutoUVSettings);
 	}
 
 
-	if ( (NewSettings.BakeResolution != CurrentSettings.BakeResolution) ||
-		 (NewSettings.BakeThickness != CurrentSettings.BakeThickness))
+	if ( (NewCombinedSettings.BakeResolution != CurrentSettings.BakeResolution) ||
+		 (NewCombinedSettings.BakeThickness != CurrentSettings.BakeThickness))
 	{
 		UE::GeometryFlow::FMeshMakeBakingCacheSettings NewBakeSettings = Generator->GetCurrentBakeCacheSettings();
-		NewBakeSettings.Dimensions = FImageDimensions((int32)NewSettings.BakeResolution, (int32)NewSettings.BakeResolution);
-		NewBakeSettings.Thickness = NewSettings.BakeThickness;
+		NewBakeSettings.Dimensions = FImageDimensions((int32)NewCombinedSettings.BakeResolution, (int32)NewCombinedSettings.BakeResolution);
+		NewBakeSettings.Thickness = NewCombinedSettings.BakeThickness;
 		Generator->UpdateBakeCacheSettings(NewBakeSettings);
 	}
 
-	if (NewSettings.ConvexTriangleCount != CurrentSettings.ConvexTriangleCount ||
-		NewSettings.bPrefilterVertices != CurrentSettings.bPrefilterVertices ||
-		NewSettings.PrefilterGridResolution != CurrentSettings.PrefilterGridResolution ||
-		NewSettings.bSimplifyPolygons != CurrentSettings.bSimplifyPolygons ||
-		NewSettings.HullTolerance != CurrentSettings.HullTolerance ||
-		NewSettings.SweepAxis != CurrentSettings.SweepAxis ||
-		NewSettings.CollisionType != CurrentSettings.CollisionType)
+	if (NewCombinedSettings.ConvexTriangleCount != CurrentSettings.ConvexTriangleCount ||
+		NewCombinedSettings.bPrefilterVertices != CurrentSettings.bPrefilterVertices ||
+		NewCombinedSettings.PrefilterGridResolution != CurrentSettings.PrefilterGridResolution ||
+		NewCombinedSettings.bSimplifyPolygons != CurrentSettings.bSimplifyPolygons ||
+		NewCombinedSettings.HullTolerance != CurrentSettings.HullTolerance ||
+		NewCombinedSettings.SweepAxis != CurrentSettings.SweepAxis ||
+		NewCombinedSettings.CollisionType != CurrentSettings.CollisionType)
 	{
 		UE::GeometryFlow::FGenerateSimpleCollisionSettings NewGenCollisionSettings = Generator->GetCurrentGenerateSimpleCollisionSettings();
-		NewGenCollisionSettings.Type = static_cast<UE::GeometryFlow::ESimpleCollisionGeometryType>(NewSettings.CollisionType);
-		NewGenCollisionSettings.ConvexHullSettings.SimplifyToTriangleCount = NewSettings.ConvexTriangleCount;
-		NewGenCollisionSettings.ConvexHullSettings.bPrefilterVertices = NewSettings.bPrefilterVertices;
-		NewGenCollisionSettings.ConvexHullSettings.PrefilterGridResolution = NewSettings.PrefilterGridResolution;
-		NewGenCollisionSettings.SweptHullSettings.bSimplifyPolygons = NewSettings.bSimplifyPolygons;
-		NewGenCollisionSettings.SweptHullSettings.HullTolerance = NewSettings.HullTolerance;
-		NewGenCollisionSettings.SweptHullSettings.SweepAxis = static_cast<FMeshSimpleShapeApproximation::EProjectedHullAxisMode>(NewSettings.SweepAxis);
+		NewGenCollisionSettings.Type = static_cast<UE::GeometryFlow::ESimpleCollisionGeometryType>(NewCombinedSettings.CollisionType);
+		NewGenCollisionSettings.ConvexHullSettings.SimplifyToTriangleCount = NewCombinedSettings.ConvexTriangleCount;
+		NewGenCollisionSettings.ConvexHullSettings.bPrefilterVertices = NewCombinedSettings.bPrefilterVertices;
+		NewGenCollisionSettings.ConvexHullSettings.PrefilterGridResolution = NewCombinedSettings.PrefilterGridResolution;
+		NewGenCollisionSettings.SweptHullSettings.bSimplifyPolygons = NewCombinedSettings.bSimplifyPolygons;
+		NewGenCollisionSettings.SweptHullSettings.HullTolerance = NewCombinedSettings.HullTolerance;
+		NewGenCollisionSettings.SweptHullSettings.SweepAxis = static_cast<FMeshSimpleShapeApproximation::EProjectedHullAxisMode>(NewCombinedSettings.SweepAxis);
 		Generator->UpdateGenerateSimpleCollisionSettings(NewGenCollisionSettings);
 	}
 
-	CurrentSettings = NewSettings;
+	CurrentSettings = NewCombinedSettings;
 }
 
 
diff --git a/Engine/Plugins/Experimental/MeshLODToolset/Source/MeshLODToolset/Private/Tools/GenerateStaticMeshLODAssetTool.cpp b/Engine/Plugins/Experimental/MeshLODToolset/Source/MeshLODToolset/Private/Tools/GenerateStaticMeshLODAssetTool.cpp
--- a/Engine/Plugins/Experimental/MeshLODToolset/Source/MeshLODToolset/Private/Tools/GenerateStaticMeshLODAssetTool.cpp
+++ b/Engine/Plugins/Experimental/MeshLODToolset/Source/MeshLODToolset/Private/Tools/GenerateStaticMeshLODAssetTool.cpp
@@ -124,6 +124,7 @@ void UGenerateStaticMeshLODAssetTool::Setup()
 	}
 
 	BasicProperties->GeneratorSettings = GenerateProcess->GetCurrentSettings();
+	BasicProperties->WatchProperty(BasicProperties->GeneratorSettings.FilterGroupLayer, [this](FName) { OnSettingsModified(); });
 	BasicProperties->WatchProperty(BasicProperties->GeneratorSettings.SolidifyVoxelResolution, [this](int) { OnSettingsModified(); });
 	BasicProperties->WatchProperty(BasicProperties->GeneratorSettings.WindingThreshold, [this](float) { OnSettingsModified(); });
 	//BasicProperties->WatchProperty(BasicProperties->GeneratorSettings.MorphologyVoxelResolution, [this](int) { OnSettingsModified(); });
@@ -182,16 +183,12 @@ void UGenerateStaticMeshLODAssetTool::Setup()
 	bPreviewValid = false;
 }
 
-
 void UGenerateStaticMeshLODAssetTool::OnSettingsModified()
 {
-	UE_LOG(LogTemp, Warning, TEXT("SETTINGS MODIFIED!"));
-
 	GenerateProcess->UpdateSettings(BasicProperties->GeneratorSettings);
 	bPreviewValid = false;
 }
 
-
 void UGenerateStaticMeshLODAssetTool::Shutdown(EToolShutdownType ShutdownType)
 {
 	BasicProperties->SaveProperties(this);
diff --git a/Engine/Plugins/Experimental/MeshLODToolset/Source/MeshLODToolset/Public/Graphs/GenerateMeshLODGraph.h b/Engine/Plugins/Experimental/MeshLODToolset/Source/MeshLODToolset/Public/Graphs/GenerateMeshLODGraph.h
--- a/Engine/Plugins/Experimental/MeshLODToolset/Source/MeshLODToolset/Public/Graphs/GenerateMeshLODGraph.h
+++ b/Engine/Plugins/Experimental/MeshLODToolset/Source/MeshLODToolset/Public/Graphs/GenerateMeshLODGraph.h
@@ -18,6 +18,13 @@
 #include "PhysicsNodes/GenerateSimpleCollisionNode.h"
 
 
+
+struct FMeshLODGraphPreFilterSettings
+{
+	FName FilterGroupLayerName = FName( "PreFilterGroups" );
+};
+
+
 class FGenerateMeshLODGraph
 {
 public:
@@ -27,6 +34,8 @@ public:
 
 	void SetSourceMesh(const FDynamicMesh3& SourceMesh);
 
+	void UpdatePreFilterSettings(const FMeshLODGraphPreFilterSettings& PreFilterSettings);
+	const FMeshLODGraphPreFilterSettings& GetCurrentPreFilterSettings() const { return CurrentPreFilterSettings; }
 
 	void UpdateSolidifySettings(const UE::GeometryFlow::FMeshSolidifySettings& SolidifySettings);
 	const UE::GeometryFlow::FMeshSolidifySettings& GetCurrentSolidifySettings() const { return CurrentSolidifySettings; }
@@ -68,6 +77,10 @@ protected:
 
 	UE::GeometryFlow::FGraph::FHandle MeshSourceNode;
 
+	UE::GeometryFlow::FGraph::FHandle FilterGroupsLayerNameNode;	// FNameSourceNode that defines name of polygroups layer
+	UE::GeometryFlow::FGraph::FHandle FilterTrianglesNode;
+	FMeshLODGraphPreFilterSettings CurrentPreFilterSettings;
+
 	UE::GeometryFlow::FGraph::FHandle SolidifyNode;
 	UE::GeometryFlow::FGraph::FHandle SolidifySettingsNode;
 	UE::GeometryFlow::FMeshSolidifySettings CurrentSolidifySettings;
diff --git a/Engine/Plugins/Experimental/MeshLODToolset/Source/MeshLODToolset/Public/Graphs/GenerateStaticMeshLODProcess.h b/Engine/Plugins/Experimental/MeshLODToolset/Source/MeshLODToolset/Public/Graphs/GenerateStaticMeshLODProcess.h
--- a/Engine/Plugins/Experimental/MeshLODToolset/Source/MeshLODToolset/Public/Graphs/GenerateStaticMeshLODProcess.h
+++ b/Engine/Plugins/Experimental/MeshLODToolset/Source/MeshLODToolset/Public/Graphs/GenerateStaticMeshLODProcess.h
@@ -64,6 +64,11 @@ struct FGenerateStaticMeshLODProcessSettings
 {
 	GENERATED_BODY()
 
+	// Filter settings
+
+	UPROPERTY(EditAnywhere, Category = DetailFilter, meta = (DisplayName = "Detail Filtering"))
+	FName FilterGroupLayer = FName(TEXT("PreFilterGroups"));
+
 	// Solidify settings
 
 	UPROPERTY(EditAnywhere, Category = Solidify, meta = (DisplayName="Voxel Resolution"))
@@ -91,7 +96,7 @@ struct FGenerateStaticMeshLODProcessSettings
 	// UV settings
 
 	UPROPERTY(EditAnywhere, Category = AutoUV, meta = (DisplayName = "AutoUV Charts", UIMin = 0, UIMax = 1000))
-	int NumAutoUVCharts = 20;
+	int NumAutoUVCharts = 0;
 
 
 	// Bake Settings
@@ -109,15 +114,15 @@ struct FGenerateStaticMeshLODProcessSettings
 	UPROPERTY(EditAnywhere, Category = Collision, meta = (DisplayName = "Collision Type"))
 	EGenerateStaticMeshLODSimpleCollisionGeometryType CollisionType = EGenerateStaticMeshLODSimpleCollisionGeometryType::ConvexHulls;
 
-	UPROPERTY(EditAnywhere, Category = Collision, meta = (DisplayName = "Convex Tri Count", 
+	UPROPERTY(EditAnywhere, Category = Collision, meta = (NoSpinbox = "true", DisplayName = "Convex Tri Count",
 														  EditCondition = "CollisionType == EGenerateStaticMeshLODSimpleCollisionGeometryType::ConvexHulls"))
 	int ConvexTriangleCount = 50;
 
 	UPROPERTY(EditAnywhere, Category = Collision, meta = (EditCondition = "CollisionType == EGenerateStaticMeshLODSimpleCollisionGeometryType::ConvexHulls"))
-	bool bPrefilterVertices = false;
+	bool bPrefilterVertices = true;
 
 	/** Grid resolution (along the maximum-length axis) */
-	UPROPERTY(EditAnywhere, Category = Collision, meta = (EditCondition = "CollisionType == EGenerateStaticMeshLODSimpleCollisionGeometryType::ConvexHulls && bPrefilterVertices", UIMin = 1, UIMax = 30))
+	UPROPERTY(EditAnywhere, Category = Collision, meta = (NoSpinbox = "true", EditCondition = "CollisionType == EGenerateStaticMeshLODSimpleCollisionGeometryType::ConvexHulls && bPrefilterVertices", UIMin = 4, UIMax = 100))
 	int PrefilterGridResolution = 10;
 
 
@@ -126,7 +131,7 @@ struct FGenerateStaticMeshLODProcessSettings
 	UPROPERTY(EditAnywhere, Category = Collision, meta = (EditCondition = "CollisionType == EGenerateStaticMeshLODSimpleCollisionGeometryType::SweptHulls"))
 	bool bSimplifyPolygons = true;
 
-	UPROPERTY(EditAnywhere, Category = Collision, meta = (UIMin = "0", UIMax = "10", ClampMin = "0", ClampMax = "100000",
+	UPROPERTY(EditAnywhere, Category = Collision, meta = (NoSpinbox = "true", UIMin = "0", UIMax = "10", ClampMin = "0", ClampMax = "100000",
 														  EditCondition = "CollisionType == EGenerateStaticMeshLODSimpleCollisionGeometryType::SweptHulls"))
 	float HullTolerance = 0.1;
 
