AutoLOD: for simple collision geometry fitting, we use polygroups to partition the mesh. Allow the user to choose which polygroup set to use by exposing the layer names in the UI.

#rnx

diff --git a/Engine/Plugins/Experimental/MeshLODToolset/Source/MeshLODToolset/Private/Graphs/GenerateMeshLODGraph.cpp b/Engine/Plugins/Experimental/MeshLODToolset/Source/MeshLODToolset/Private/Graphs/GenerateMeshLODGraph.cpp
--- a/Engine/Plugins/Experimental/MeshLODToolset/Source/MeshLODToolset/Private/Graphs/GenerateMeshLODGraph.cpp
+++ b/Engine/Plugins/Experimental/MeshLODToolset/Source/MeshLODToolset/Private/Graphs/GenerateMeshLODGraph.cpp
@@ -266,6 +266,13 @@ void FGenerateMeshLODGraph::UpdateGenerateSimpleCollisionSettings(const FGenerat
 }
 
 
+void FGenerateMeshLODGraph::UpdateCollisionGroupLayerName(const FName& NewCollisionGroupLayerName)
+{
+	CollisionGroupLayerName = NewCollisionGroupLayerName;
+	UpdateSourceNodeValue<FNameSourceNode>(*Graph, GroupLayerNameNode, CollisionGroupLayerName);
+}
+
+
 void FGenerateMeshLODGraph::EvaluateResult(
 	FDynamicMesh3& ResultMesh,
 	FMeshTangentsd& ResultTangents,
@@ -494,6 +501,10 @@ void FGenerateMeshLODGraph::BuildGraph()
 	Graph->InferConnection(FilterTrianglesNode, DecomposeMeshForCollisionNode);
 	Graph->InferConnection(IgnoreGroupsForCollisionNode, DecomposeMeshForCollisionNode);
 
+	GroupLayerNameNode = Graph->AddNodeOfType<FNameSourceNode>(TEXT("GroupLayerNameNode"));
+	Graph->AddConnection(GroupLayerNameNode, FNameSourceNode::OutParamValue(), 
+						 DecomposeMeshForCollisionNode, FMakeTriangleSetsFromGroupsNode::InParamGroupLayer());
+
 	GenerateSimpleCollisionNode = Graph->AddNodeOfType<FGenerateSimpleCollisionNode>(TEXT("GenerateSimpleCollision"));
 	Graph->InferConnection(FilterTrianglesNode, GenerateSimpleCollisionNode);
 	Graph->InferConnection(DecomposeMeshForCollisionNode, GenerateSimpleCollisionNode);
@@ -567,6 +578,8 @@ void FGenerateMeshLODGraph::BuildGraph()
 	IgnoreGroupsForCollision.AppendSet({ 0 });
 	UpdateSettingsSourceNodeValue(*Graph, IgnoreGroupsForCollisionNode, IgnoreGroupsForCollision);
 
+	UpdateCollisionGroupLayerName(CollisionGroupLayerName);
+
 	FGenerateSimpleCollisionSettings GenSimpleCollisionSettings;
 	UpdateGenerateSimpleCollisionSettings(GenSimpleCollisionSettings);
 
diff --git a/Engine/Plugins/Experimental/MeshLODToolset/Source/MeshLODToolset/Private/Graphs/GenerateStaticMeshLODProcess.cpp b/Engine/Plugins/Experimental/MeshLODToolset/Source/MeshLODToolset/Private/Graphs/GenerateStaticMeshLODProcess.cpp
--- a/Engine/Plugins/Experimental/MeshLODToolset/Source/MeshLODToolset/Private/Graphs/GenerateStaticMeshLODProcess.cpp
+++ b/Engine/Plugins/Experimental/MeshLODToolset/Source/MeshLODToolset/Private/Graphs/GenerateStaticMeshLODProcess.cpp
@@ -289,6 +289,11 @@ void UGenerateStaticMeshLODProcess::UpdateSettings(const FGenerateStaticMeshLODP
 		Generator->UpdateBakeCacheSettings(NewBakeSettings);
 	}
 
+	if (NewCombinedSettings.CollisionGroupLayerName != CurrentSettings.CollisionGroupLayerName)
+	{
+		Generator->UpdateCollisionGroupLayerName(NewCombinedSettings.CollisionGroupLayerName);
+	}
+
 	if (NewCombinedSettings.ConvexTriangleCount != CurrentSettings.ConvexTriangleCount ||
 		NewCombinedSettings.bPrefilterVertices != CurrentSettings.bPrefilterVertices ||
 		NewCombinedSettings.PrefilterGridResolution != CurrentSettings.PrefilterGridResolution ||
diff --git a/Engine/Plugins/Experimental/MeshLODToolset/Source/MeshLODToolset/Private/Tools/GenerateStaticMeshLODAssetTool.cpp b/Engine/Plugins/Experimental/MeshLODToolset/Source/MeshLODToolset/Private/Tools/GenerateStaticMeshLODAssetTool.cpp
--- a/Engine/Plugins/Experimental/MeshLODToolset/Source/MeshLODToolset/Private/Tools/GenerateStaticMeshLODAssetTool.cpp
+++ b/Engine/Plugins/Experimental/MeshLODToolset/Source/MeshLODToolset/Private/Tools/GenerateStaticMeshLODAssetTool.cpp
@@ -40,6 +40,8 @@
 #include "ModelingOperators.h"
 #include "MeshOpPreviewHelpers.h"
 #include "Generators/GridBoxMeshGenerator.h"
+#include "Polygroups/PolygroupSet.h"
+#include "Polygroups/PolygroupUtil.h"
 
 
 #if WITH_EDITOR
@@ -127,6 +129,7 @@ namespace GenerateStaticMeshLODAssetLocals
 			TUniquePtr<FGenerateStaticMeshLODAssetOperatorOp> Op = MakeUnique<FGenerateStaticMeshLODAssetOperatorOp>();		
 			Op->GenerateProcess = AutoLODTool->GenerateProcess;
 			Op->GeneratorSettings = AutoLODTool->BasicProperties->GeneratorSettings;
+			Op->GeneratorSettings.CollisionGroupLayerName = AutoLODTool->BasicProperties->CollisionGroupLayerName;
 			Op->SetResultTransform(ResultTransform);
 			return Op;
 		}
@@ -193,15 +196,6 @@ void UGenerateStaticMeshLODAssetTool::Setup()
 
 	UInteractiveTool::Setup();
 
-	BasicProperties = NewObject<UGenerateStaticMeshLODAssetToolProperties>(this);
-	AddToolPropertySource(BasicProperties);
-	BasicProperties->RestoreProperties(this);
-
-	BasicProperties->OutputName = AssetGenerationUtil::GetComponentAssetBaseName(ComponentTargets[0]->GetOwnerComponent());
-
-	BasicProperties->GeneratedSuffix = TEXT("_AutoLOD");
-
-
 	GetToolManager()->DisplayMessage(
 		LOCTEXT("OnStartStaticMeshLODAssetTool", "This tool creates a new LOD asset"),
 		EToolMessageLevel::UserNotification);
@@ -216,6 +210,11 @@ void UGenerateStaticMeshLODAssetTool::Setup()
 		GenerateProcess->Initialize(StaticMesh);		// Must happen on main thread
 	}
 
+	BasicProperties = NewObject<UGenerateStaticMeshLODAssetToolProperties>(this);
+	AddToolPropertySource(BasicProperties);
+	BasicProperties->RestoreProperties(this);
+	BasicProperties->OutputName = AssetGenerationUtil::GetComponentAssetBaseName(ComponentTargets[0]->GetOwnerComponent());
+	BasicProperties->GeneratedSuffix = TEXT("_AutoLOD");
 	BasicProperties->GeneratorSettings = GenerateProcess->GetCurrentSettings();
 	BasicProperties->WatchProperty(BasicProperties->GeneratorSettings.FilterGroupLayer, [this](FName) { OnSettingsModified(); });
 	BasicProperties->WatchProperty(BasicProperties->GeneratorSettings.SolidifyVoxelResolution, [this](int) { OnSettingsModified(); });
@@ -233,6 +232,10 @@ void UGenerateStaticMeshLODAssetTool::Setup()
 	BasicProperties->WatchProperty(BasicProperties->GeneratorSettings.bSimplifyPolygons, [this](bool) { OnSettingsModified(); });
 	BasicProperties->WatchProperty(BasicProperties->GeneratorSettings.HullTolerance, [this](float) { OnSettingsModified(); });
 	BasicProperties->WatchProperty(BasicProperties->GeneratorSettings.SweepAxis, [this](EGenerateStaticMeshLODProjectedHullAxisMode) { OnSettingsModified(); });
+	
+	// Collision layer name property
+	BasicProperties->WatchProperty(BasicProperties->CollisionGroupLayerName, [this](FName) { OnSettingsModified(); });
+	BasicProperties->InitializeGroupLayers(&(GenerateProcess->GetSourceMesh()));
 
 	FBoxSphereBounds Bounds = StaticMeshComponent->Bounds;
 	FTransform PreviewTransform = SourceComponent->GetWorldTransform();
diff --git a/Engine/Plugins/Experimental/MeshLODToolset/Source/MeshLODToolset/Public/Graphs/GenerateMeshLODGraph.h b/Engine/Plugins/Experimental/MeshLODToolset/Source/MeshLODToolset/Public/Graphs/GenerateMeshLODGraph.h
--- a/Engine/Plugins/Experimental/MeshLODToolset/Source/MeshLODToolset/Public/Graphs/GenerateMeshLODGraph.h
+++ b/Engine/Plugins/Experimental/MeshLODToolset/Source/MeshLODToolset/Public/Graphs/GenerateMeshLODGraph.h
@@ -55,6 +55,7 @@ public:
 	void UpdateGenerateSimpleCollisionSettings(const UE::GeometryFlow::FGenerateSimpleCollisionSettings& SimpleCollisionSettings);
 	const UE::GeometryFlow::FGenerateSimpleCollisionSettings& GetCurrentGenerateSimpleCollisionSettings() const { return CurrentGenerateSimpleCollisionSettings; }
 
+	void UpdateCollisionGroupLayerName(const FName& CollisionGroupLayerName);
 
 
 
@@ -118,6 +119,8 @@ protected:
 	UE::GeometryFlow::FGraph::FHandle BakeNormalMapNode;
 	UE::GeometryFlow::FGraph::FHandle BakeNormalMapSettingsNode;
 
+	UE::GeometryFlow::FGraph::FHandle GroupLayerNameNode;
+
 	struct FBakeTextureGraphInfo
 	{
 		int32 Index;
@@ -137,4 +140,6 @@ protected:
 	UE::GeometryFlow::FGraph::FHandle MeshOutputNode;
 	UE::GeometryFlow::FGraph::FHandle TangentsOutputNode;
 
+	FName CollisionGroupLayerName = TEXT("Default");
+	
 };
diff --git a/Engine/Plugins/Experimental/MeshLODToolset/Source/MeshLODToolset/Public/Graphs/GenerateStaticMeshLODProcess.h b/Engine/Plugins/Experimental/MeshLODToolset/Source/MeshLODToolset/Public/Graphs/GenerateStaticMeshLODProcess.h
--- a/Engine/Plugins/Experimental/MeshLODToolset/Source/MeshLODToolset/Public/Graphs/GenerateStaticMeshLODProcess.h
+++ b/Engine/Plugins/Experimental/MeshLODToolset/Source/MeshLODToolset/Public/Graphs/GenerateStaticMeshLODProcess.h
@@ -107,13 +107,19 @@ struct FGenerateStaticMeshLODProcessSettings
 	UPROPERTY(EditAnywhere, Category = Baking, meta = (DisplayName = "Bake Thickness"))
 	float BakeThickness = 5.0f;
 
+	// Transient property, not set directly by the user. The user controls a CollisionGroupLayerName dropdown property
+	// on the Tool and that value is copied here.
+	UPROPERTY(meta = (TransientToolProperty))
+	FName CollisionGroupLayerName = TEXT("Default");
 
 
-	// Convex Hull Settings
+	// Simple collision generator settings
 
 	UPROPERTY(EditAnywhere, Category = Collision, meta = (DisplayName = "Collision Type"))
 	EGenerateStaticMeshLODSimpleCollisionGeometryType CollisionType = EGenerateStaticMeshLODSimpleCollisionGeometryType::ConvexHulls;
 
+	// Convex Hull Settings
+
 	UPROPERTY(EditAnywhere, Category = Collision, meta = (NoSpinbox = "true", DisplayName = "Convex Tri Count",
 														  EditCondition = "CollisionType == EGenerateStaticMeshLODSimpleCollisionGeometryType::ConvexHulls"))
 	int ConvexTriangleCount = 50;
diff --git a/Engine/Plugins/Experimental/MeshLODToolset/Source/MeshLODToolset/Public/Tools/GenerateStaticMeshLODAssetTool.h b/Engine/Plugins/Experimental/MeshLODToolset/Source/MeshLODToolset/Public/Tools/GenerateStaticMeshLODAssetTool.h
--- a/Engine/Plugins/Experimental/MeshLODToolset/Source/MeshLODToolset/Public/Tools/GenerateStaticMeshLODAssetTool.h
+++ b/Engine/Plugins/Experimental/MeshLODToolset/Source/MeshLODToolset/Public/Tools/GenerateStaticMeshLODAssetTool.h
@@ -75,6 +75,44 @@ public:
 	UPROPERTY(EditAnywhere, Category = Settings)
 	FGenerateStaticMeshLODProcessSettings GeneratorSettings;
 
+	// ------------
+	// Code copied from UPolygroupLayersProperties
+
+	UPROPERTY(EditAnywhere, Category = Settings, meta = (GetOptions = GetGroupLayersFunc))
+	FName CollisionGroupLayerName = TEXT("Default");
+
+	// this function is called provide set of available group layers
+	UFUNCTION()
+	TArray<FString> GetGroupLayersFunc()
+	{
+		return GroupLayersList;
+	}
+
+	// internal list used to implement above
+	UPROPERTY(meta = (TransientToolProperty))
+	TArray<FString> GroupLayersList;
+
+	void InitializeGroupLayers(const FDynamicMesh3* Mesh)
+	{
+		GroupLayersList.Reset();
+		GroupLayersList.Add(TEXT("Default"));		// always have standard group
+		if (Mesh->Attributes())
+		{
+			for (int32 k = 0; k < Mesh->Attributes()->NumPolygroupLayers(); k++)
+			{
+				FName Name = Mesh->Attributes()->GetPolygroupLayer(k)->GetName();
+				GroupLayersList.Add(Name.ToString());
+			}
+		}
+
+		if (GroupLayersList.Contains(CollisionGroupLayerName.ToString()) == false)		// discard restored value if it doesn't apply
+		{
+			CollisionGroupLayerName = FName(GroupLayersList[0]);
+		}
+	}
+
+	// ------------
+
 	UPROPERTY(VisibleAnywhere, Category = Previews)
 	TArray<UTexture2D*> PreviewTextures;
 
