Add a Thicken node to GeometryFlow. Accepts vertex weights and a scalar and moves vertices in the normal direction accordingly.

diff --git a/Engine/Plugins/Experimental/GeometryFlow/Source/GeometryFlowMeshProcessing/Private/MeshProcessingNodes/MeshThickenNode.cpp b/Engine/Plugins/Experimental/GeometryFlow/Source/GeometryFlowMeshProcessing/Private/MeshProcessingNodes/MeshThickenNode.cpp
new file mode 100644
--- /dev/null
+++ b/Engine/Plugins/Experimental/GeometryFlow/Source/GeometryFlowMeshProcessing/Private/MeshProcessingNodes/MeshThickenNode.cpp
@@ -0,0 +1,50 @@
+// Copyright Epic Games, Inc. All Rights Reserved.
+
+#include "MeshProcessingNodes/MeshThickenNode.h"
+#include "DynamicMesh3.h"
+#include "Operations/DisplaceMesh.h"
+#include "MeshNormals.h"
+
+using namespace UE::GeometryFlow;
+
+void FMeshThickenNode::ProcessMesh(
+	const FNamedDataMap& DatasIn,
+	const FMeshThickenSettings& SettingsIn,
+	const FDynamicMesh3& MeshIn,
+	FDynamicMesh3& MeshOut) 
+{
+	TSafeSharedPtr<IData> WeightMapMeshArg = DatasIn.FindData(InParamWeightMap());
+	FWeightMap WeightMap;
+	WeightMapMeshArg->GetDataCopy<FWeightMap>(WeightMap, (int)EMeshProcessingDataTypes::WeightMap);
+
+	MeshOut = MeshIn;
+	ApplyThicken(MeshOut, SettingsIn, WeightMap.Weights);
+}
+
+
+void FMeshThickenNode::ProcessMeshInPlace(
+	const FNamedDataMap& DatasIn,
+	const FMeshThickenSettings& Settings,
+	FDynamicMesh3& MeshInOut) 
+{
+	TSafeSharedPtr<IData> WeightMapMeshArg = DatasIn.FindData(InParamWeightMap());
+	FWeightMap WeightMap;
+	WeightMapMeshArg->GetDataCopy<FWeightMap>(WeightMap, (int)EMeshProcessingDataTypes::WeightMap);
+
+	ApplyThicken(MeshInOut, Settings, WeightMap.Weights);
+}
+
+
+void FMeshThickenNode::ApplyThicken(FDynamicMesh3& Mesh, const FMeshThickenSettings& Settings, const TArray<float>& VertexWeights)
+{
+	if (VertexWeights.Num() == 0)
+	{
+		return;
+	}
+
+	// TODO: Accept existing normals and/or allow user to choose type of normal to compute
+	FMeshNormals Normals(&Mesh);
+	Normals.ComputeVertexNormals();
+
+	FDisplaceMesh::DisplaceMeshWithVertexWeights(Mesh, Normals, VertexWeights, Settings.ThickenAmount);
+}
diff --git a/Engine/Plugins/Experimental/GeometryFlow/Source/GeometryFlowMeshProcessing/Public/MeshProcessingNodes/MeshProcessingDataTypes.h b/Engine/Plugins/Experimental/GeometryFlow/Source/GeometryFlowMeshProcessing/Public/MeshProcessingNodes/MeshProcessingDataTypes.h
--- a/Engine/Plugins/Experimental/GeometryFlow/Source/GeometryFlowMeshProcessing/Public/MeshProcessingNodes/MeshProcessingDataTypes.h
+++ b/Engine/Plugins/Experimental/GeometryFlow/Source/GeometryFlowMeshProcessing/Public/MeshProcessingNodes/MeshProcessingDataTypes.h
@@ -25,6 +25,7 @@ enum class EMeshProcessingDataTypes
 
 	SolidifySettings = (int)EDataTypes::BaseMeshProcessingTypes + 100,
 	VoxMorphologyOpSettings = (int)EDataTypes::BaseMeshProcessingTypes + 101,
+	ThickenSettings = (int)EDataTypes::BaseMeshProcessingTypes + 102,
 
 	SimplifySettings = (int)EDataTypes::BaseMeshProcessingTypes + 110,
 	NormalFlowSettings = (int)EDataTypes::BaseMeshProcessingTypes + 111,
diff --git a/Engine/Plugins/Experimental/GeometryFlow/Source/GeometryFlowMeshProcessing/Public/MeshProcessingNodes/MeshThickenNode.h b/Engine/Plugins/Experimental/GeometryFlow/Source/GeometryFlowMeshProcessing/Public/MeshProcessingNodes/MeshThickenNode.h
new file mode 100644
--- /dev/null
+++ b/Engine/Plugins/Experimental/GeometryFlow/Source/GeometryFlowMeshProcessing/Public/MeshProcessingNodes/MeshThickenNode.h
@@ -0,0 +1,60 @@
+// Copyright Epic Games, Inc. All Rights Reserved.
+
+#pragma once
+
+#include "MeshProcessingNodes/MeshProcessingBaseNodes.h"
+#include "MeshProcessingNodes/MeshProcessingDataTypes.h"
+#include "DataTypes/WeightMapData.h"
+
+namespace UE
+{
+namespace GeometryFlow
+{
+
+struct GEOMETRYFLOWMESHPROCESSING_API FMeshThickenSettings
+{
+	DECLARE_GEOMETRYFLOW_DATA_TYPE_IDENTIFIER(EMeshProcessingDataTypes::ThickenSettings);
+
+	float ThickenAmount = 1.0f;
+};
+GEOMETRYFLOW_DECLARE_SETTINGS_TYPES(FMeshThickenSettings, Thicken);
+
+
+class GEOMETRYFLOWMESHPROCESSING_API FMeshThickenNode : public TProcessMeshWithSettingsBaseNode<FMeshThickenSettings>
+{
+public:
+
+	static const FString InParamWeightMap() { return TEXT("WeightMap"); }
+
+	FMeshThickenNode()
+	{
+		AddInput(InParamWeightMap(), MakeUnique<FWeightMapInput>());
+	}
+
+	void CheckAdditionalInputs(const FNamedDataMap& DatasIn,
+							   bool& bRecomputeRequired,
+							   bool& bAllInputsValid) override
+	{
+		FindAndUpdateInputForEvaluate(InParamWeightMap(), DatasIn, bRecomputeRequired, bAllInputsValid);
+	}
+
+	void ProcessMesh(
+		const FNamedDataMap& DatasIn,
+		const FMeshThickenSettings& SettingsIn,
+		const FDynamicMesh3& MeshIn,
+		FDynamicMesh3& MeshOut) override;
+
+	void ProcessMeshInPlace(
+		const FNamedDataMap& DatasIn,
+		const FMeshThickenSettings& Settings,
+		FDynamicMesh3& MeshInOut) override;
+
+protected:
+
+	void ApplyThicken(FDynamicMesh3& MeshInOut, const FMeshThickenSettings& Settings, const TArray<float>& VertexWeights);
+	
+};
+
+
+}	// end namespace GeometryFlow
+}	// end namespace UE
diff --git a/Engine/Plugins/Experimental/GeometryProcessing/Source/DynamicMesh/Public/Operations/DisplaceMesh.h b/Engine/Plugins/Experimental/GeometryProcessing/Source/DynamicMesh/Public/Operations/DisplaceMesh.h
new file mode 100644
--- /dev/null
+++ b/Engine/Plugins/Experimental/GeometryProcessing/Source/DynamicMesh/Public/Operations/DisplaceMesh.h
@@ -0,0 +1,43 @@
+// Copyright Epic Games, Inc. All Rights Reserved.
+
+#pragma once
+
+#include "MathUtil.h"
+#include "VectorTypes.h"
+#include "GeometryTypes.h"
+#include "MeshNormals.h"
+#include "Async/ParallelFor.h"
+
+class FDisplaceMesh
+{
+
+public:
+
+	inline static void DisplaceMeshWithVertexWeights(FDynamicMesh3& Mesh,
+													 const FMeshNormals& Normals,
+													 const TArray<float>& VertexWeights,
+													 float Intensity)
+	{
+		check(VertexWeights.Num() >= Mesh.MaxVertexID());
+		check(Normals.GetNormals().Num() >= Mesh.MaxVertexID());
+
+		int NumVertices = Mesh.MaxVertexID();
+		
+		ParallelFor(NumVertices, [&Mesh, &VertexWeights, &Normals, Intensity](int32 VertexID)
+		{
+			if (!Mesh.IsVertex(VertexID)) 
+			{ 
+				return; 
+			}
+
+			float Offset = VertexWeights[VertexID];
+			FVector3d NewPosition = Mesh.GetVertexRef(VertexID) + (Offset * Intensity * Normals[VertexID]);
+			Mesh.SetVertex_NoTimeStampUpdate(VertexID, NewPosition);
+		});
+		Mesh.IncrementTimeStamps(1, true, false);
+	}
+
+
+	// TODO: Refactor DisplaceMeshTool to use this class
+
+};
diff --git a/Engine/Plugins/Experimental/MeshLODToolset/Source/MeshLODToolset/Private/Graphs/GenerateMeshLODGraph.cpp b/Engine/Plugins/Experimental/MeshLODToolset/Source/MeshLODToolset/Private/Graphs/GenerateMeshLODGraph.cpp
--- a/Engine/Plugins/Experimental/MeshLODToolset/Source/MeshLODToolset/Private/Graphs/GenerateMeshLODGraph.cpp
+++ b/Engine/Plugins/Experimental/MeshLODToolset/Source/MeshLODToolset/Private/Graphs/GenerateMeshLODGraph.cpp
@@ -7,6 +7,7 @@
 #include "GeometryFlowGraphUtil.h"
 #include "BaseNodes/TransferNode.h"
 
+#include "MeshProcessingNodes/MeshThickenNode.h"
 #include "MeshProcessingNodes/MeshSolidifyNode.h"
 #include "MeshProcessingNodes/MeshVoxMorphologyNode.h"
 #include "MeshProcessingNodes/MeshSimplifyNode.h"
@@ -265,6 +266,18 @@ void FGenerateMeshLODGraph::UpdateGenerateSimpleCollisionSettings(const FGenerat
 	CurrentGenerateSimpleCollisionSettings = GenSimpleCollisionSettings;
 }
 
+void FGenerateMeshLODGraph::UpdateThickenWeightMap(const TArray<float>& ThickenWeightMap)
+{
+	FWeightMap WeightMap;
+	WeightMap.Weights = ThickenWeightMap;
+	UpdateSourceNodeValue<FWeightMapSourceNode>(*Graph, ThickenWeightMapNode, WeightMap);
+}
+
+void FGenerateMeshLODGraph::UpdateThickenSettings(const UE::GeometryFlow::FMeshThickenSettings& ThickenSettings)
+{
+	UpdateSettingsSourceNodeValue(*Graph, ThickenSettingsNode, ThickenSettings);
+	CurrentThickenSettings = ThickenSettings;
+}
 
 void FGenerateMeshLODGraph::UpdateCollisionGroupLayerName(const FName& NewCollisionGroupLayerName)
 {
@@ -420,8 +433,17 @@ void FGenerateMeshLODGraph::BuildGraph()
 
 	// generating low-poly mesh
 
+	// optionally thicken some parts of the mesh before solidifying
+	ThickenNode = Graph->AddNodeOfType<FMeshThickenNode>(TEXT("Thicken"));
+	ThickenWeightMapNode = Graph->AddNodeOfType<FWeightMapSourceNode>(TEXT("ThickenWeightMapNode"));
+	ThickenSettingsNode = Graph->AddNodeOfType<FThickenSettingsSourceNode>(TEXT("ThickenSettingsSource"));
+
+	Graph->InferConnection(ThickenWeightMapNode, ThickenNode);
+	Graph->InferConnection(ThickenSettingsNode, ThickenNode);
+	Graph->InferConnection(FilterTrianglesNode, ThickenNode);
+
 	SolidifyNode = Graph->AddNodeOfType<FSolidifyMeshNode>(TEXT("Solidify"));
-	Graph->InferConnection(FilterTrianglesNode, SolidifyNode);
+	Graph->InferConnection(ThickenNode, SolidifyNode);
 	SolidifySettingsNode = Graph->AddNodeOfType<FSolidifySettingsSourceNode>(TEXT("SolidifySettings"));
 	Graph->InferConnection(SolidifySettingsNode, SolidifyNode);
 
@@ -583,6 +605,9 @@ void FGenerateMeshLODGraph::BuildGraph()
 	FGenerateSimpleCollisionSettings GenSimpleCollisionSettings;
 	UpdateGenerateSimpleCollisionSettings(GenSimpleCollisionSettings);
 
+	TArray<float> Weights;
+	UpdateThickenWeightMap(Weights);
+
 
 	//FString GraphDump = Graph->DebugDumpGraph([](TSafeSharedPtr<FNode> Node)
 	//{
diff --git a/Engine/Plugins/Experimental/MeshLODToolset/Source/MeshLODToolset/Private/Graphs/GenerateStaticMeshLODProcess.cpp b/Engine/Plugins/Experimental/MeshLODToolset/Source/MeshLODToolset/Private/Graphs/GenerateStaticMeshLODProcess.cpp
--- a/Engine/Plugins/Experimental/MeshLODToolset/Source/MeshLODToolset/Private/Graphs/GenerateStaticMeshLODProcess.cpp
+++ b/Engine/Plugins/Experimental/MeshLODToolset/Source/MeshLODToolset/Private/Graphs/GenerateStaticMeshLODProcess.cpp
@@ -28,6 +28,7 @@
 #include "Materials/MaterialInstanceConstant.h"
 #include "Factories/MaterialInstanceConstantFactoryNew.h"
 #include "Materials/MaterialInstanceDynamic.h"
+#include "WeightMapUtil.h"
 
 #include "Engine/Classes/Engine/StaticMesh.h"
 #include "Engine/Classes/Components/StaticMeshComponent.h"
@@ -247,6 +248,32 @@ void UGenerateStaticMeshLODProcess::UpdateSettings(const FGenerateStaticMeshLODP
 		Generator->UpdatePreFilterSettings(NewPreFilterSettings);
 	}
 
+	if (NewCombinedSettings.ThickenAmount != CurrentSettings.ThickenAmount)
+	{
+		UE::GeometryFlow::FMeshThickenSettings NewThickenSettings = Generator->GetCurrentThickenSettings();
+		NewThickenSettings.ThickenAmount = NewCombinedSettings.ThickenAmount;
+		Generator->UpdateThickenSettings(NewThickenSettings);
+	}
+
+	if (NewCombinedSettings.ThickenWeightMapName != CurrentSettings.ThickenWeightMapName)
+	{
+		const FMeshDescription* SourceMeshDescription = SourceStaticMesh->GetMeshDescription(0);
+		FIndexedWeightMap1f WeightMap;
+		float DefaultValue = 0.0f;
+		bool bOK = UE::WeightMaps::GetVertexWeightMap(SourceMeshDescription, NewCombinedSettings.ThickenWeightMapName, WeightMap, DefaultValue);
+
+		if (bOK)
+		{
+			Generator->UpdateThickenWeightMap(WeightMap.Values);
+		}
+		else
+		{
+			Generator->UpdateThickenWeightMap(TArray<float>());
+		}
+	}
+
+	
+
 	bool bSharedVoxelResolutionChanged = (NewCombinedSettings.SolidifyVoxelResolution != CurrentSettings.SolidifyVoxelResolution);
 	if ( bSharedVoxelResolutionChanged
 		|| (NewCombinedSettings.WindingThreshold != CurrentSettings.WindingThreshold))
diff --git a/Engine/Plugins/Experimental/MeshLODToolset/Source/MeshLODToolset/Private/Tools/GenerateStaticMeshLODAssetTool.cpp b/Engine/Plugins/Experimental/MeshLODToolset/Source/MeshLODToolset/Private/Tools/GenerateStaticMeshLODAssetTool.cpp
--- a/Engine/Plugins/Experimental/MeshLODToolset/Source/MeshLODToolset/Private/Tools/GenerateStaticMeshLODAssetTool.cpp
+++ b/Engine/Plugins/Experimental/MeshLODToolset/Source/MeshLODToolset/Private/Tools/GenerateStaticMeshLODAssetTool.cpp
@@ -227,6 +227,8 @@ void UGenerateStaticMeshLODAssetTool::Setup()
 	BasicProperties->GeneratedSuffix = TEXT("_AutoLOD");
 	BasicProperties->GeneratorSettings = GenerateProcess->GetCurrentSettings();
 	BasicProperties->WatchProperty(BasicProperties->GeneratorSettings.FilterGroupLayer, [this](FName) { OnSettingsModified(); });
+	BasicProperties->WatchProperty(BasicProperties->GeneratorSettings.ThickenWeightMapName, [this](FName) { OnSettingsModified(); });
+	BasicProperties->WatchProperty(BasicProperties->GeneratorSettings.ThickenAmount, [this](float) { OnSettingsModified(); });
 	BasicProperties->WatchProperty(BasicProperties->GeneratorSettings.SolidifyVoxelResolution, [this](int) { OnSettingsModified(); });
 	BasicProperties->WatchProperty(BasicProperties->GeneratorSettings.WindingThreshold, [this](float) { OnSettingsModified(); });
 	//BasicProperties->WatchProperty(BasicProperties->GeneratorSettings.MorphologyVoxelResolution, [this](int) { OnSettingsModified(); });
diff --git a/Engine/Plugins/Experimental/MeshLODToolset/Source/MeshLODToolset/Public/Graphs/GenerateMeshLODGraph.h b/Engine/Plugins/Experimental/MeshLODToolset/Source/MeshLODToolset/Public/Graphs/GenerateMeshLODGraph.h
--- a/Engine/Plugins/Experimental/MeshLODToolset/Source/MeshLODToolset/Public/Graphs/GenerateMeshLODGraph.h
+++ b/Engine/Plugins/Experimental/MeshLODToolset/Source/MeshLODToolset/Public/Graphs/GenerateMeshLODGraph.h
@@ -12,6 +12,7 @@
 #include "MeshProcessingNodes/MeshSolidifyNode.h"
 #include "MeshProcessingNodes/MeshVoxMorphologyNode.h"
 #include "MeshProcessingNodes/MeshSimplifyNode.h"
+#include "MeshProcessingNodes/MeshThickenNode.h"
 #include "MeshProcessingNodes/MeshDeleteTrianglesNode.h"
 #include "MeshProcessingNodes/MeshAutoGenerateUVsNode.h"
 #include "DataTypes/MeshImageBakingData.h"
@@ -55,9 +56,11 @@ public:
 	void UpdateGenerateSimpleCollisionSettings(const UE::GeometryFlow::FGenerateSimpleCollisionSettings& SimpleCollisionSettings);
 	const UE::GeometryFlow::FGenerateSimpleCollisionSettings& GetCurrentGenerateSimpleCollisionSettings() const { return CurrentGenerateSimpleCollisionSettings; }
 
+	void UpdateThickenWeightMap(const TArray<float>& ThickenWeightMap);
 	void UpdateCollisionGroupLayerName(const FName& CollisionGroupLayerName);
 
-
+	void UpdateThickenSettings(const UE::GeometryFlow::FMeshThickenSettings& ThickenSettings);
+	const UE::GeometryFlow::FMeshThickenSettings& GetCurrentThickenSettings() const { return CurrentThickenSettings; }
 
 	void EvaluateResult(
 		FDynamicMesh3& ResultMesh,
@@ -119,6 +122,11 @@ protected:
 	UE::GeometryFlow::FGraph::FHandle BakeNormalMapNode;
 	UE::GeometryFlow::FGraph::FHandle BakeNormalMapSettingsNode;
 
+	UE::GeometryFlow::FGraph::FHandle ThickenNode;
+	UE::GeometryFlow::FGraph::FHandle ThickenSettingsNode;
+	UE::GeometryFlow::FGraph::FHandle ThickenWeightMapNode;
+	UE::GeometryFlow::FMeshThickenSettings CurrentThickenSettings;
+
 	UE::GeometryFlow::FGraph::FHandle GroupLayerNameNode;
 
 	struct FBakeTextureGraphInfo
diff --git a/Engine/Plugins/Experimental/MeshLODToolset/Source/MeshLODToolset/Public/Graphs/GenerateStaticMeshLODProcess.h b/Engine/Plugins/Experimental/MeshLODToolset/Source/MeshLODToolset/Public/Graphs/GenerateStaticMeshLODProcess.h
--- a/Engine/Plugins/Experimental/MeshLODToolset/Source/MeshLODToolset/Public/Graphs/GenerateStaticMeshLODProcess.h
+++ b/Engine/Plugins/Experimental/MeshLODToolset/Source/MeshLODToolset/Public/Graphs/GenerateStaticMeshLODProcess.h
@@ -69,6 +69,14 @@ struct FGenerateStaticMeshLODProcessSettings
 	UPROPERTY(EditAnywhere, Category = DetailFilter, meta = (DisplayName = "Detail Filtering"))
 	FName FilterGroupLayer = FName(TEXT("PreFilterGroups"));
 
+	// Thicken settings
+
+	UPROPERTY(EditAnywhere, Category = DetailFilter, meta = (DisplayName = "Thicken weight map name"))
+	FName ThickenWeightMapName = FName(TEXT("ThickenWeightMap"));
+
+	UPROPERTY(EditAnywhere, Category = DetailFilter)
+	float ThickenAmount = 0.0f;
+
 	// Solidify settings
 
 	UPROPERTY(EditAnywhere, Category = Solidify, meta = (DisplayName="Voxel Resolution"))
