AutoLOD: for simple collision geometry fitting, we use polygroups to partition the mesh. Allow the user to choose which polygroup set to use by exposing the layer names in the UI.

#rnx

diff --git a/Engine/Plugins/Experimental/MeshLODToolset/Source/MeshLODToolset/Private/Tools/GenerateStaticMeshLODAssetTool.cpp b/Engine/Plugins/Experimental/MeshLODToolset/Source/MeshLODToolset/Private/Tools/GenerateStaticMeshLODAssetTool.cpp
--- a/Engine/Plugins/Experimental/MeshLODToolset/Source/MeshLODToolset/Private/Tools/GenerateStaticMeshLODAssetTool.cpp
+++ b/Engine/Plugins/Experimental/MeshLODToolset/Source/MeshLODToolset/Private/Tools/GenerateStaticMeshLODAssetTool.cpp
@@ -196,15 +196,6 @@ void UGenerateStaticMeshLODAssetTool::Setup()
 
 	UInteractiveTool::Setup();
 
-	BasicProperties = NewObject<UGenerateStaticMeshLODAssetToolProperties>(this);
-	AddToolPropertySource(BasicProperties);
-	BasicProperties->RestoreProperties(this);
-
-	BasicProperties->OutputName = AssetGenerationUtil::GetComponentAssetBaseName(ComponentTargets[0]->GetOwnerComponent());
-
-	BasicProperties->GeneratedSuffix = TEXT("_AutoLOD");
-
-
 	SetToolDisplayName(LOCTEXT("ToolName", "Generate LOD"));
 	GetToolManager()->DisplayMessage(
 		LOCTEXT("OnStartStaticMeshLODAssetTool", "Create a new LOD asset"),
