GeometryProcessing: Added FMeshAttributeTransfer class, currently supports transferring triangle material IDs from one mesh to another, via nearest-point queries

GeometryFlow: Added FCompactMeshNode, FComputeMeshPerVertexOverlayNormalsNode (convenient), and FTransferMeshMaterialIDsNode that applies FMeshAttributeTransfer based on two input meshes. Graph InferConnection() now is more clever and can infer connections in cases where a potentially-ambiguous connection is resolved based on what is already connected (ie if two mesh inputs, and one is assigned, infer that second one is intended)

LODToolset: LODGraph now transfers materials and computes per-vertex normals before Simplify node, and Compacts afterwards (required for other downstream processes).
LODProcess now handles multiple materials, and can re-use input materials if they do not require a baked texture or normal map. In repeat applications, previously-auto-generated StaticMesh Materials are detected and removed from the material set. Normals and Tangents are explicitly computed for stored HiRes source and LOD0, and build settings updated appropriately.

#rnx

diff --git a/Engine/Plugins/Experimental/GeometryFlow/Source/GeometryFlowCore/Private/GeometryFlowGraph.cpp b/Engine/Plugins/Experimental/GeometryFlow/Source/GeometryFlowCore/Private/GeometryFlowGraph.cpp
--- a/Engine/Plugins/Experimental/GeometryFlow/Source/GeometryFlowCore/Private/GeometryFlowGraph.cpp
+++ b/Engine/Plugins/Experimental/GeometryFlow/Source/GeometryFlowCore/Private/GeometryFlowGraph.cpp
@@ -116,12 +116,17 @@ EGeometryFlowResult FGraph::InferConnection(FHandle FromNodeHandle, FHandle ToNo
 		int32 OutputType = Output->GetDataType();
 		ToNode->EnumerateInputs([&](const FString& InputName, const TUniquePtr<INodeInput>& Input)
 		{
-			int32 InputType = Input->GetDataType();
-			if (OutputType == InputType)
+			// if input already has a connection, we cannot add another one and can skip it
+			FConnection ExistingConnection;
+			if (FindConnectionForInput(ToNodeHandle, InputName, ExistingConnection) != EGeometryFlowResult::Ok)
 			{
-				TotalMatchesFound++;
-				FromOutputName = OutputName;
-				ToInputName = InputName;
+				int32 InputType = Input->GetDataType();
+				if (OutputType == InputType)
+				{
+					TotalMatchesFound++;
+					FromOutputName = OutputName;
+					ToInputName = InputName;
+				}
 			}
 		});
 	});
diff --git a/Engine/Plugins/Experimental/GeometryFlow/Source/GeometryFlowMeshProcessing/Public/MeshProcessingNodes/CompactMeshNode.h b/Engine/Plugins/Experimental/GeometryFlow/Source/GeometryFlowMeshProcessing/Public/MeshProcessingNodes/CompactMeshNode.h
new file mode 100644
--- /dev/null
+++ b/Engine/Plugins/Experimental/GeometryFlow/Source/GeometryFlowMeshProcessing/Public/MeshProcessingNodes/CompactMeshNode.h
@@ -0,0 +1,48 @@
+// Copyright Epic Games, Inc. All Rights Reserved.
+
+#pragma once
+
+#include "MeshProcessingNodes/MeshProcessingBaseNodes.h"
+#include "MeshProcessingNodes/MeshProcessingDataTypes.h"
+
+namespace UE
+{
+namespace GeometryFlow
+{
+
+
+/**
+ * FCompactMeshNode compacts the input mesh, to remove gaps in the vertex/triangle indexing. This node can be applied in-place.
+ */
+class FCompactMeshNode : public FProcessMeshBaseNode
+{
+public:
+	FCompactMeshNode()
+	{
+		// we can mutate input mesh
+		ConfigureInputFlags(InParamMesh(), FNodeInputFlags::Transformable());
+	}
+
+	virtual void ProcessMesh(
+		const FNamedDataMap& DatasIn,
+		const FDynamicMesh3& MeshIn,
+		FDynamicMesh3& MeshOut) override
+	{
+		MeshOut.CompactCopy(MeshIn);
+	}
+
+	virtual void ProcessMeshInPlace(
+		const FNamedDataMap& DatasIn,
+		FDynamicMesh3& MeshInOut) override
+	{
+		MeshInOut.CompactInPlace();
+	}
+
+};
+
+
+
+
+
+}	// end namespace GeometryFlow
+}	//
\ No newline at end of file
diff --git a/Engine/Plugins/Experimental/GeometryFlow/Source/GeometryFlowMeshProcessing/Public/MeshProcessingNodes/MeshNormalsNodes.h b/Engine/Plugins/Experimental/GeometryFlow/Source/GeometryFlowMeshProcessing/Public/MeshProcessingNodes/MeshNormalsNodes.h
--- a/Engine/Plugins/Experimental/GeometryFlow/Source/GeometryFlowMeshProcessing/Public/MeshProcessingNodes/MeshNormalsNodes.h
+++ b/Engine/Plugins/Experimental/GeometryFlow/Source/GeometryFlowMeshProcessing/Public/MeshProcessingNodes/MeshNormalsNodes.h
@@ -40,7 +40,9 @@ struct FMeshNormalsSettings
 GEOMETRYFLOW_DECLARE_SETTINGS_TYPES(FMeshNormalsSettings, Normals);
 
 
-
+/**
+ * Recompute Normals overlay for input Mesh. Can apply in-place.
+ */
 class FComputeMeshNormalsNode : public TProcessMeshWithSettingsBaseNode<FMeshNormalsSettings>
 {
 public:
@@ -114,6 +116,48 @@ public:
 
 
 
+/**
+ * Recompute per-vertex normals in Normals Overlay for input mesh. Can apply in-place.
+ */
+class FComputeMeshPerVertexOverlayNormalsNode : public FProcessMeshBaseNode
+{
+public:
+	FComputeMeshPerVertexOverlayNormalsNode()
+	{
+		// we can mutate input mesh
+		ConfigureInputFlags(InParamMesh(), FNodeInputFlags::Transformable());
+	}
+
+	virtual void ProcessMesh(
+		const FNamedDataMap& DatasIn,
+		const FDynamicMesh3& MeshIn,
+		FDynamicMesh3& MeshOut) override
+	{
+		MeshOut = MeshIn;
+		if (MeshOut.HasAttributes() == false)
+		{
+			MeshOut.EnableAttributes();
+		}
+		FDynamicMeshNormalOverlay* Normals = MeshOut.Attributes()->PrimaryNormals();
+		FMeshNormals::InitializeOverlayToPerVertexNormals(Normals, false);
+	}
+
+	virtual void ProcessMeshInPlace(
+		const FNamedDataMap& DatasIn,
+		FDynamicMesh3& MeshInOut) override
+	{
+		if (MeshInOut.HasAttributes() == false)
+		{
+			MeshInOut.EnableAttributes();
+		}
+		FDynamicMeshNormalOverlay* Normals = MeshInOut.Attributes()->PrimaryNormals();
+		FMeshNormals::InitializeOverlayToPerVertexNormals(Normals, false);
+	}
+
+};
+
+
+
 
 
 }	// end namespace GeometryFlow
diff --git a/Engine/Plugins/Experimental/GeometryFlow/Source/GeometryFlowMeshProcessing/Public/MeshProcessingNodes/TransferMeshMaterialIDsNode.h b/Engine/Plugins/Experimental/GeometryFlow/Source/GeometryFlowMeshProcessing/Public/MeshProcessingNodes/TransferMeshMaterialIDsNode.h
new file mode 100644
--- /dev/null
+++ b/Engine/Plugins/Experimental/GeometryFlow/Source/GeometryFlowMeshProcessing/Public/MeshProcessingNodes/TransferMeshMaterialIDsNode.h
@@ -0,0 +1,80 @@
+// Copyright Epic Games, Inc. All Rights Reserved.
+
+#pragma once
+
+#include "MeshProcessingNodes/MeshProcessingBaseNodes.h"
+#include "MeshProcessingNodes/MeshProcessingDataTypes.h"
+#include "Operations/MeshAttributeTransfer.h"
+
+namespace UE
+{
+namespace GeometryFlow
+{
+
+/**
+ * FTransferMeshMaterialIDsNode transfers Material IDs from a SourceMesh to the Input mesh and returns in the Output mesh (can be applied in-place)
+ */
+class FTransferMeshMaterialIDsNode : public FProcessMeshBaseNode
+{
+public:
+	static const FString InParamMaterialSourceMesh() { return TEXT("MaterialSourceMesh"); }
+
+public:
+	FTransferMeshMaterialIDsNode()
+	{
+		AddInput(InParamMaterialSourceMesh(), MakeUnique<FDynamicMeshInput>());
+
+		// we can mutate input mesh
+		ConfigureInputFlags(InParamMesh(), FNodeInputFlags::Transformable());
+	}
+
+	virtual void CheckAdditionalInputs(const FNamedDataMap& DatasIn, bool& bRecomputeRequired, bool& bAllInputsValid) override
+	{
+		FindAndUpdateInputForEvaluate(InParamMaterialSourceMesh(), DatasIn, bRecomputeRequired, bAllInputsValid);
+	}
+
+
+	virtual void ProcessMesh(
+		const FNamedDataMap& DatasIn,
+		const FDynamicMesh3& MeshIn,
+		FDynamicMesh3& MeshOut) override
+	{
+		MeshOut = MeshIn;
+		TransferMaterialIDs(DatasIn, &MeshOut);
+	}
+
+	virtual void ProcessMeshInPlace(
+		const FNamedDataMap& DatasIn,
+		FDynamicMesh3& MeshInOut) override
+	{
+		TransferMaterialIDs(DatasIn, &MeshInOut);
+	}
+
+
+	virtual void TransferMaterialIDs(const FNamedDataMap& DatasIn, FDynamicMesh3* TargetMesh)
+	{
+		TSafeSharedPtr<IData> MaterialSourceMeshArg = DatasIn.FindData(InParamMaterialSourceMesh());
+		const FDynamicMesh3& MaterialSourceMesh = MaterialSourceMeshArg->GetDataConstRef<FDynamicMesh3>((int)EMeshProcessingDataTypes::DynamicMesh);
+
+		if (TargetMesh->HasAttributes() == false)
+		{
+			TargetMesh->EnableAttributes();
+		}
+		if (TargetMesh->Attributes()->HasMaterialID() == false)
+		{
+			TargetMesh->Attributes()->EnableMaterialID();
+		}
+
+		FMeshAttributeTransfer Transfer(&MaterialSourceMesh, TargetMesh);
+		Transfer.TransferType = EMeshAttributeTransferType::MaterialID;
+		Transfer.Apply();
+	}
+
+};
+
+
+
+
+
+}	// end namespace GeometryFlow
+}	//
\ No newline at end of file
diff --git a/Engine/Plugins/Experimental/GeometryProcessing/Source/DynamicMesh/Private/Operations/MeshAttributeTransfer.cpp b/Engine/Plugins/Experimental/GeometryProcessing/Source/DynamicMesh/Private/Operations/MeshAttributeTransfer.cpp
new file mode 100644
--- /dev/null
+++ b/Engine/Plugins/Experimental/GeometryProcessing/Source/DynamicMesh/Private/Operations/MeshAttributeTransfer.cpp
@@ -0,0 +1,123 @@
+// Copyright Epic Games, Inc. All Rights Reserved.
+
+#include "Operations/MeshAttributeTransfer.h"
+#include "DynamicMesh3.h"
+#include "DynamicMeshAttributeSet.h"
+#include "TriangleTypes.h"
+#include "Util/IndexUtil.h"
+#include "Async/ParallelFor.h"
+
+
+
+
+FMeshAttributeTransfer::FMeshAttributeTransfer(const FDynamicMesh3* SourceMeshIn, FDynamicMesh3* TargetMeshIn)
+	: SourceMesh(SourceMeshIn), TargetMesh(TargetMeshIn)
+{
+	ensure(SourceMeshIn);
+	ensure(TargetMeshIn);
+}
+
+
+bool FMeshAttributeTransfer::Apply()
+{
+	if (SourceMesh == nullptr)
+	{
+		Errors.Add(EMeshAttributeTransferError::InvalidSource);
+		return false;
+	}
+	if (TargetMesh == nullptr)
+	{
+		Errors.Add(EMeshAttributeTransferError::InvalidTarget);
+		return false;
+	}
+
+	if ( TransferType == EMeshAttributeTransferType::MaterialID)
+	{
+		return Apply_MaterialID();
+	}
+	else
+	{
+		Errors.Add(EMeshAttributeTransferError::InvalidOperationForInputs);
+		return false;
+	}
+
+}
+
+
+
+template<typename AllocType>
+int32 SelectMostFrequentElement(TArray<int32, AllocType>& Samples)
+{
+	Samples.Sort();
+	int32 N = Samples.Num();
+	int32 MaxVotes = 0;
+	int32 MaxVoteValue = Samples[0];
+	int32 CurValue = Samples[0];
+	int32 CurCount = 1;
+	for (int32 k = 1; k < N; ++k)
+	{
+		if (Samples[k] == CurValue)
+		{
+			CurCount++;
+		}
+		else
+		{
+			if (CurCount > MaxVotes)
+			{
+				MaxVotes = CurCount;
+				MaxVoteValue = CurValue;
+			}
+			CurCount = 1;
+			CurValue = Samples[k];
+		}
+	}
+	return MaxVoteValue;
+}
+
+
+
+bool FMeshAttributeTransfer::Apply_MaterialID()
+{
+	const FDynamicMeshMaterialAttribute* SourceMaterialID = SourceMesh->HasAttributes() ? SourceMesh->Attributes()->GetMaterialID() : nullptr;
+	if (SourceMaterialID == nullptr)
+	{
+		Errors.Add(EMeshAttributeTransferError::SourceMissingAttribute);
+		return false;
+	}
+	FDynamicMeshMaterialAttribute* TargetMaterialID = TargetMesh->HasAttributes() ? TargetMesh->Attributes()->GetMaterialID() : nullptr;
+	if (TargetMaterialID == nullptr)
+	{
+		Errors.Add(EMeshAttributeTransferError::TargetMissingAttribute);
+		return false;
+	}
+	
+	SourceSpatial = MakeUnique<FDynamicMeshAABBTree3>(SourceMesh, true);
+
+	static const FVector3d BarySamplePoints[7] = {
+		{1, 0, 0}, { 0,1,0 }, { 0,0,1 },
+		{ 0.5,0.5,0 }, {0.5,0,0.5}, {0,0.5,0.5},
+		{1.0/3.0, 1.0/3.0, 1.0/3.0}
+	};
+	constexpr int32 NumSamplePoints = 7;		// must change Samples array below too!
+
+	// generate initial assignments
+	ParallelFor(TargetMesh->MaxTriangleID(), [&](int32 tid)
+	{
+		FTriangle3d TargetTriangle;
+		TargetMesh->GetTriVertices(tid, TargetTriangle.V[0], TargetTriangle.V[1], TargetTriangle.V[2]);
+
+		TArray<int32, TFixedAllocator<7>> Samples;
+		for (int32 k = 0; k < NumSamplePoints; ++k)
+		{
+			FVector3d SamplePoint = TargetTriangle.BarycentricPoint(BarySamplePoints[k]);
+			double DistSqr;
+			int32 SourceTriID = SourceSpatial->FindNearestTriangle(SamplePoint, DistSqr);
+			Samples.Add(SourceMaterialID->GetValue(SourceTriID));
+		}
+
+		int32 BestValue = SelectMostFrequentElement(Samples);
+		TargetMaterialID->SetValue(tid, BestValue);
+	});
+
+	return true;
+}
\ No newline at end of file
diff --git a/Engine/Plugins/Experimental/GeometryProcessing/Source/DynamicMesh/Public/Operations/MeshAttributeTransfer.h b/Engine/Plugins/Experimental/GeometryProcessing/Source/DynamicMesh/Public/Operations/MeshAttributeTransfer.h
new file mode 100644
--- /dev/null
+++ b/Engine/Plugins/Experimental/GeometryProcessing/Source/DynamicMesh/Public/Operations/MeshAttributeTransfer.h
@@ -0,0 +1,61 @@
+// Copyright Epic Games, Inc. All Rights Reserved.
+
+#pragma once
+
+#include "MathUtil.h"
+#include "VectorTypes.h"
+#include "DynamicMeshAABBTree3.h"
+
+class FDynamicMesh3;
+
+
+enum class EMeshAttributeTransferType
+{
+	MaterialID
+};
+
+enum class EMeshAttributeTransferError
+{
+	InvalidSource,
+	InvalidTarget,
+	SourceMissingAttribute,
+	TargetMissingAttribute,
+	InvalidOperationForInputs
+};
+
+/**
+ * FMeshAttributeTransfer transfers attributes from a SourceMesh to a TargetMesh
+ */
+class DYNAMICMESH_API FMeshAttributeTransfer
+{
+public:
+	/** The mesh that we are transferring from */
+	const FDynamicMesh3* SourceMesh;
+
+	/** The mesh that we are transferring to */
+	FDynamicMesh3* TargetMesh;
+
+	/** What is being transferred */
+	EMeshAttributeTransferType TransferType = EMeshAttributeTransferType::MaterialID;
+
+
+	TArray<EMeshAttributeTransferError> Errors;
+
+public:
+	FMeshAttributeTransfer(const FDynamicMesh3* SourceMeshIn, FDynamicMesh3* TargetMeshIn);
+
+	virtual ~FMeshAttributeTransfer() {}
+
+	/**
+	 * Run the operation and modify .Mesh
+	 * @return true if the algorithm succeeds
+	 */
+	virtual bool Apply();
+
+protected:
+
+	TUniquePtr<FDynamicMeshAABBTree3> SourceSpatial;
+
+
+	bool Apply_MaterialID();
+};
\ No newline at end of file
diff --git a/Engine/Plugins/Experimental/MeshLODToolset/Source/MeshLODToolset/Private/Graphs/GenerateMeshLODGraph.cpp b/Engine/Plugins/Experimental/MeshLODToolset/Source/MeshLODToolset/Private/Graphs/GenerateMeshLODGraph.cpp
--- a/Engine/Plugins/Experimental/MeshLODToolset/Source/MeshLODToolset/Private/Graphs/GenerateMeshLODGraph.cpp
+++ b/Engine/Plugins/Experimental/MeshLODToolset/Source/MeshLODToolset/Private/Graphs/GenerateMeshLODGraph.cpp
@@ -12,6 +12,8 @@
 #include "MeshProcessingNodes/MeshVoxMorphologyNode.h"
 #include "MeshProcessingNodes/MeshSimplifyNode.h"
 #include "MeshProcessingNodes/MeshDeleteTrianglesNode.h"
+#include "MeshProcessingNodes/CompactMeshNode.h"
+#include "MeshProcessingNodes/TransferMeshMaterialIDsNode.h"
 
 #include "MeshProcessingNodes/MeshNormalsNodes.h"
 #include "MeshProcessingNodes/MeshTangentsNodes.h"
@@ -452,13 +454,25 @@ void FGenerateMeshLODGraph::BuildGraph()
 	MorphologySettingsNode = Graph->AddNodeOfType<FVoxClosureSettingsSourceNode>(TEXT("ClosureSettings"));
 	Graph->InferConnection(MorphologySettingsNode, MorphologyNode);
 
+	// todo: if we only have one material ID we can skip this...
+	FGraph::FHandle MatIDTransferNode = Graph->AddNodeOfType<FTransferMeshMaterialIDsNode>(TEXT("TransferMaterialIDs"));
+	Graph->AddConnection(MeshSourceNode, FDynamicMeshSourceNode::OutParamValue(), MatIDTransferNode, FTransferMeshMaterialIDsNode::InParamMaterialSourceMesh());
+	Graph->InferConnection(MorphologyNode, MatIDTransferNode);
+
+	// need to compute valid normals before Simplify, Morphology node does not necessarily do it
+	FGraph::FHandle PerVertexNormalsNode = Graph->AddNodeOfType<FComputeMeshPerVertexOverlayNormalsNode>(TEXT("PerVertexNormals"));
+	Graph->InferConnection(MatIDTransferNode, PerVertexNormalsNode);
+
 	SimplifyNode = Graph->AddNodeOfType<FSimplifyMeshNode>(TEXT("Simplify"));
-	Graph->InferConnection(MorphologyNode, SimplifyNode);
+	Graph->InferConnection(PerVertexNormalsNode, SimplifyNode);
 	SimplifySettingsNode = Graph->AddNodeOfType<FSimplifySettingsSourceNode>(TEXT("SimplifySettings"));
 	Graph->InferConnection(SimplifySettingsNode, SimplifyNode);
 
+	FGraph::FHandle CompactNode = Graph->AddNodeOfType<FCompactMeshNode>(TEXT("Compact"));
+	Graph->InferConnection(SimplifyNode, CompactNode);
+
 	NormalsNode = Graph->AddNodeOfType<FComputeMeshNormalsNode>(TEXT("Normals"));
-	Graph->InferConnection(SimplifyNode, NormalsNode);
+	Graph->InferConnection(CompactNode, NormalsNode);
 	NormalsSettingsNode = Graph->AddNodeOfType<FNormalsSettingsSourceNode>(TEXT("NormalsSettings"));
 	Graph->InferConnection(NormalsSettingsNode, NormalsNode);
 
@@ -560,10 +574,11 @@ void FGenerateMeshLODGraph::BuildGraph()
 	UpdateMorphologySettings(MorphologySettings);
 
 	FMeshSimplifySettings SimplifySettings;
-	SimplifySettings.bDiscardAttributes = true;
-	SimplifySettings.SimplifyType = EMeshSimplifyType::VolumePreserving;
+	SimplifySettings.bDiscardAttributes = false;
+	SimplifySettings.SimplifyType = EMeshSimplifyType::AttributeAware;
 	SimplifySettings.TargetType = EMeshSimplifyTargetType::TriangleCount;
 	SimplifySettings.TargetCount = 500;
+	SimplifySettings.MaterialBorderConstraints = EEdgeRefineFlags::NoFlip;
 	UpdateSimplifySettings(SimplifySettings);
 
 	FMeshNormalsSettings NormalsSettings;
diff --git a/Engine/Plugins/Experimental/MeshLODToolset/Source/MeshLODToolset/Private/Graphs/GenerateStaticMeshLODProcess.cpp b/Engine/Plugins/Experimental/MeshLODToolset/Source/MeshLODToolset/Private/Graphs/GenerateStaticMeshLODProcess.cpp
--- a/Engine/Plugins/Experimental/MeshLODToolset/Source/MeshLODToolset/Private/Graphs/GenerateStaticMeshLODProcess.cpp
+++ b/Engine/Plugins/Experimental/MeshLODToolset/Source/MeshLODToolset/Private/Graphs/GenerateStaticMeshLODProcess.cpp
@@ -11,6 +11,7 @@
 
 #include "AssetUtils/Texture2DUtil.h"
 #include "AssetUtils/Texture2DBuilder.h"
+#include "AssetUtils/MeshDescriptionUtil.h"
 
 #include "Physics/PhysicsDataCollection.h"
 
@@ -30,9 +31,10 @@
 #include "Materials/MaterialInstanceDynamic.h"
 #include "WeightMapUtil.h"
 
-#include "Engine/Classes/Engine/StaticMesh.h"
-#include "Engine/Classes/Components/StaticMeshComponent.h"
-#include "Engine/Classes/PhysicsEngine/BodySetup.h"
+#include "Engine/StaticMesh.h"
+#include "Components/StaticMeshComponent.h"
+#include "PhysicsEngine/BodySetup.h"
+#include "StaticMeshAttributes.h"
 
 #define LOCTEXT_NAMESPACE "UGenerateStaticMeshLODProcess"
 
@@ -65,16 +67,24 @@ bool UGenerateStaticMeshLODProcess::Initialize(UStaticMesh* StaticMeshIn)
 
 	SourceStaticMesh = StaticMeshIn;
 
-	bool bHasHiResSource = SourceStaticMesh->IsHiResMeshDescriptionValid();
-	const FMeshDescription* SourceMeshDescription =
-		(bHasHiResSource) ? SourceStaticMesh->GetHiResMeshDescription() : SourceStaticMesh->GetMeshDescription(0);
+	bUsingHiResSource = SourceStaticMesh->IsHiResMeshDescriptionValid();
+	const FMeshDescription* UseSourceMeshDescription =
+		(bUsingHiResSource) ? SourceStaticMesh->GetHiResMeshDescription() : SourceStaticMesh->GetMeshDescription(0);
+	SourceMeshDescription = MakeShared<FMeshDescription>();
+	*SourceMeshDescription = *UseSourceMeshDescription;
+
+	// if not the high-res source, compute autogenerated normals/tangents
+	if (bUsingHiResSource == false)
+	{
+		UE::MeshDescription::InitializeAutoGeneratedAttributes(*SourceMeshDescription, SourceStaticMesh, 0);
+	}
 
 	// start async mesh-conversion
 	SourceMesh.Clear();
 	TFuture<void> ConvertMesh = Async(GenerateSMLODAsyncExecTarget, [&]()
 	{
 		FMeshDescriptionToDynamicMesh GetSourceMesh;
-		GetSourceMesh.Convert(SourceMeshDescription, SourceMesh);
+		GetSourceMesh.Convert(SourceMeshDescription.Get(), SourceMesh);
 	});
 
 	// get list of source materials and find all the input texture params
@@ -91,6 +101,9 @@ bool UGenerateStaticMeshLODProcess::Initialize(UStaticMesh* StaticMeshIn)
 			continue;
 		}
 
+		SourceMaterials[mi].bHasNormalMap = false;
+		SourceMaterials[mi].bHasTexturesToBake = false;
+
 		TArray<FMaterialParameterInfo> ParameterInfo;
 		TArray<FGuid> ParameterIds;
 		MaterialInterface->GetAllTextureParameterInfo(ParameterInfo, ParameterIds);
@@ -106,30 +119,89 @@ bool UGenerateStaticMeshLODProcess::Initialize(UStaticMesh* StaticMeshIn)
 				if (Cast<UTexture2D>(CurTexture) != nullptr)
 				{
 					FTextureInfo TexInfo;
-					TexInfo.SourceTexture = Cast<UTexture2D>(CurTexture);
+					TexInfo.Texture = Cast<UTexture2D>(CurTexture);
 					TexInfo.ParameterName = ParamName;
 
-					TexInfo.bIsNormalMap = TexInfo.SourceTexture->IsNormalMap();
-					//TexInfo.bIsDefaultTexture = TexInfo.SourceTexture->IsDefaultTexture();
-					TexInfo.bIsDefaultTexture = UEditorAssetLibrary::GetPathNameForLoadedAsset(TexInfo.SourceTexture).StartsWith(TEXT("/Engine/"));
+					TexInfo.bIsNormalMap = TexInfo.Texture->IsNormalMap();
+					SourceMaterials[mi].bHasNormalMap |= TexInfo.bIsNormalMap;
 
-					UE_LOG(LogTemp, Warning, TEXT("SOURCE TEX PATH IS %s - DefaultTex %d - SRGB %d"), *UEditorAssetLibrary::GetPathNameForLoadedAsset(TexInfo.SourceTexture),
-						TexInfo.bIsDefaultTexture?1:0, TexInfo.SourceTexture->SRGB);
+					//TexInfo.bIsDefaultTexture = TexInfo.Texture->IsDefaultTexture();
+					TexInfo.bIsDefaultTexture = UEditorAssetLibrary::GetPathNameForLoadedAsset(TexInfo.Texture).StartsWith(TEXT("/Engine/"));
 
-					TexInfo.bShouldBakeTexture = (TexInfo.bIsNormalMap == false && TexInfo.bIsDefaultTexture == false);
+					UE_LOG(LogTemp, Warning, TEXT("SOURCE TEX PATH IS %s - DefaultTex %d - SRGB %d"), *UEditorAssetLibrary::GetPathNameForLoadedAsset(TexInfo.Texture),
+						TexInfo.bIsDefaultTexture?1:0, TexInfo.Texture->SRGB);
 
+					TexInfo.bShouldBakeTexture = (TexInfo.bIsNormalMap == false && TexInfo.bIsDefaultTexture == false);
 					if (TexInfo.bShouldBakeTexture)
 					{
 						ReadJobs.Add(FReadTextureJob{ mi, SourceMaterials[mi].SourceTextures.Num() });
+
+						SourceMaterials[mi].bHasTexturesToBake = true;
 					}
 
 					SourceMaterials[mi].SourceTextures.Add(TexInfo);
 				}
 			}
 		}
+
+		// if material does not have a normal map parameter or any textures we want to bake, we can just re-use it
+		SourceMaterials[mi].bIsReusable = (SourceMaterials[mi].bHasNormalMap == false && SourceMaterials[mi].bHasTexturesToBake == false);
+	}
+
+
+	// If we have hi-res source we can discard any materials that are only used on the previously-generated LOD0.
+	// We cannot explicitly tag the materials as being generated so we infer, ie we assume a material was generated
+	// if it is only used in LOD0 and not HiRes
+	if (bUsingHiResSource)
+	{
+		// have to wait for SourceMesh to finish converting
+		ConvertMesh.Wait();
+
+		const FMeshDescription* LOD0MeshDescription = SourceStaticMesh->GetMeshDescription(0);
+		FMeshDescriptionToDynamicMesh GetLOD0Mesh;
+		FDynamicMesh3 LOD0Mesh;
+		GetLOD0Mesh.Convert(LOD0MeshDescription, LOD0Mesh);
+		const FDynamicMeshMaterialAttribute* SourceMaterialIDs = SourceMesh.Attributes()->GetMaterialID();
+		const FDynamicMeshMaterialAttribute* LOD0MaterialIDs = LOD0Mesh.Attributes()->GetMaterialID();
+		if (ensure(SourceMaterialIDs != nullptr && LOD0MaterialIDs != nullptr))
+		{
+			int32 NumMaterials = SourceMaterials.Num();
+			TArray<bool> IsBaseMaterial;
+			IsBaseMaterial.Init(false, NumMaterials);
+			TArray<bool> IsLOD0Material;
+			IsLOD0Material.Init(false, NumMaterials);
+
+			for (int32 tid : SourceMesh.TriangleIndicesItr())
+			{
+				int32 MatIdx = SourceMaterialIDs->GetValue(tid);
+				if (MatIdx >= 0 && MatIdx < NumMaterials)
+				{
+					IsBaseMaterial[MatIdx] = true;
+				}
+			}
+			for (int32 tid : LOD0Mesh.TriangleIndicesItr())
+			{
+				int32 MatIdx = LOD0MaterialIDs->GetValue(tid);
+				if (MatIdx >= 0 && MatIdx < NumMaterials)
+				{
+					IsLOD0Material[MatIdx] = true;
+				}
+			}
+
+			for (int32 k = 0; k < NumMaterials; ++k)
+			{
+				if (IsLOD0Material[k] == true && IsBaseMaterial[k] == false)
+				{
+					SourceMaterials[k].bIsPreviouslyGeneratedMaterial = true;
+					SourceMaterials[k].bHasTexturesToBake = false;
+					SourceMaterials[k].bIsReusable = false;
+				}
+			}
+		}
 	}
 
 
+
 	// extract all the texture params
 	// TODO: this triggers a checkSlow in the serialization code when it runs async. Find out why. Jira: UETOOL-2985
 	//TFuture<void> ReadTextures = Async(GenerateSMLODAsyncExecTarget, [&]()
@@ -138,15 +210,20 @@ bool UGenerateStaticMeshLODProcess::Initialize(UStaticMesh* StaticMeshIn)
 	//	{
 	//		FReadTextureJob job = ReadJobs[ji];
 	//		FTextureInfo& TexInfo = SourceMaterials[job.MatIndex].SourceTextures[job.TexIndex];
-	//		UE::AssetUtils::ReadTexture(TexInfo.SourceTexture, TexInfo.Dimensions, TexInfo.Image);
+	//		UE::AssetUtils::ReadTexture(TexInfo.Texture, TexInfo.Dimensions, TexInfo.Image);
 	//	});
 	//});
 
 	// single-thread path
 	for (FReadTextureJob job : ReadJobs)
 	{
-		FTextureInfo& TexInfo = SourceMaterials[job.MatIndex].SourceTextures[job.TexIndex];
-		UE::AssetUtils::ReadTexture(TexInfo.SourceTexture, TexInfo.Dimensions, TexInfo.Image);
+		// only read textures that are from materials we are going to possibly bake
+		FSourceMaterialInfo& SourceMaterial = SourceMaterials[job.MatIndex];
+		if (SourceMaterial.bIsPreviouslyGeneratedMaterial == false && SourceMaterial.bIsReusable == false)
+		{
+			FTextureInfo& TexInfo = SourceMaterials[job.MatIndex].SourceTextures[job.TexIndex];
+			UE::AssetUtils::ReadTexture(TexInfo.Texture, TexInfo.Dimensions, TexInfo.Image);
+		}
 	}
 
 
@@ -186,17 +263,20 @@ bool UGenerateStaticMeshLODProcess::InitializeGenerator()
 	Generator = MakeUnique<FGenerateMeshLODGraph>();
 	Generator->BuildGraph();
 
-	// initialize source textures
+	// initialize source textures that need to be baked
 	SourceTextureToDerivedTexIndex.Reset();
-	for (const FMaterialInfo& MatInfo : SourceMaterials)
+	for (const FSourceMaterialInfo& MatInfo : SourceMaterials)
 	{
-		for (const FTextureInfo& TexInfo : MatInfo.SourceTextures)
+		if ( MatInfo.bIsPreviouslyGeneratedMaterial == false && MatInfo.bIsReusable == false && MatInfo.bHasTexturesToBake)
 		{
-			if (TexInfo.bShouldBakeTexture &&
-				(SourceTextureToDerivedTexIndex.Contains(TexInfo.SourceTexture) == false))
+			for (const FTextureInfo& TexInfo : MatInfo.SourceTextures)
 			{
-				int32 NewIndex = Generator->AppendTextureBakeNode(TexInfo.Image, TexInfo.SourceTexture->GetName());
-				SourceTextureToDerivedTexIndex.Add(TexInfo.SourceTexture, NewIndex);
+				if (TexInfo.bShouldBakeTexture &&
+					(SourceTextureToDerivedTexIndex.Contains(TexInfo.Texture) == false))
+				{
+					int32 NewIndex = Generator->AppendTextureBakeNode(TexInfo.Image, TexInfo.Texture->GetName());
+					SourceTextureToDerivedTexIndex.Add(TexInfo.Texture, NewIndex);
+				}
 			}
 		}
 	}
@@ -257,10 +337,9 @@ void UGenerateStaticMeshLODProcess::UpdateSettings(const FGenerateStaticMeshLODP
 
 	if (NewCombinedSettings.ThickenWeightMapName != CurrentSettings.ThickenWeightMapName)
 	{
-		const FMeshDescription* SourceMeshDescription = SourceStaticMesh->GetMeshDescription(0);
 		FIndexedWeightMap1f WeightMap;
 		float DefaultValue = 0.0f;
-		bool bOK = UE::WeightMaps::GetVertexWeightMap(SourceMeshDescription, NewCombinedSettings.ThickenWeightMapName, WeightMap, DefaultValue);
+		bool bOK = UE::WeightMaps::GetVertexWeightMap(SourceMeshDescription.Get(), NewCombinedSettings.ThickenWeightMapName, WeightMap, DefaultValue);
 
 		if (bOK)
 		{
@@ -378,18 +457,33 @@ bool UGenerateStaticMeshLODProcess::ComputeDerivedSourceData(FProgressCancel* Pr
 		return false;
 	}
 
-	// copy all materials for now...we are going to replace all the images though, and
-	// should not copy those?
-	DerivedMaterials = SourceMaterials;
+	// copy materials
+	int32 NumMaterials = SourceMaterials.Num();
+	DerivedMaterials.SetNum(NumMaterials);
+	for (int32 mi = 0; mi < NumMaterials; ++mi)
+	{
+		DerivedMaterials[mi].SourceMaterialIndex = mi;
+		DerivedMaterials[mi].bUseSourceMaterialDirectly = (SourceMaterials[mi].bIsReusable || SourceMaterials[mi].bIsPreviouslyGeneratedMaterial);
+		
+		if (DerivedMaterials[mi].bUseSourceMaterialDirectly)
+		{
+			DerivedMaterials[mi].DerivedMaterial = SourceMaterials[mi].SourceMaterial;
+		}
+		else
+		{
+			// TODO this is a lot of wasted overhead, we do not need to copy images here for example
+			DerivedMaterials[mi].DerivedTextures = SourceMaterials[mi].SourceTextures;
+		}
+	}
 
-	// update texture data in derived materials
-	for (FMaterialInfo& MatInfo : DerivedMaterials)
+	// update texture image data in derived materials
+	for (FDerivedMaterialInfo& MatInfo : DerivedMaterials)
 	{
-		for (FTextureInfo& TexInfo : MatInfo.SourceTextures)
+		for (FTextureInfo& TexInfo : MatInfo.DerivedTextures)
 		{
-			if (SourceTextureToDerivedTexIndex.Contains(TexInfo.SourceTexture))
+			if (SourceTextureToDerivedTexIndex.Contains(TexInfo.Texture))
 			{
-				int32 BakedTexIndex = SourceTextureToDerivedTexIndex[TexInfo.SourceTexture];
+				int32 BakedTexIndex = SourceTextureToDerivedTexIndex[TexInfo.Texture];
 				const TUniquePtr<UE::GeometryFlow::FTextureImage>& DerivedTex = DerivedTextureImages[BakedTexIndex];
 				TexInfo.Dimensions = DerivedTex->Image.GetDimensions();
 
@@ -419,16 +513,21 @@ void UGenerateStaticMeshLODProcess::GetDerivedMaterialsPreview(FPreviewMaterials
 	TMap<UTexture2D*, UTexture2D*> SourceToPreviewTexMap;
 	for (int32 mi = 0; mi < NumMaterials; ++mi)
 	{
-		const FMaterialInfo& SourceMaterialInfo = SourceMaterials[mi];
-		const FMaterialInfo& DerivedMaterialInfo = DerivedMaterials[mi];
+		const FSourceMaterialInfo& SourceMaterialInfo = SourceMaterials[mi];
+		if (SourceMaterialInfo.bIsReusable || SourceMaterialInfo.bIsPreviouslyGeneratedMaterial)
+		{
+			continue;
+		}
+
+		const FDerivedMaterialInfo& DerivedMaterialInfo = DerivedMaterials[mi];
 
 		int32 NumTextures = SourceMaterialInfo.SourceTextures.Num();
-		check(DerivedMaterialInfo.SourceTextures.Num() == NumTextures);
+		check(DerivedMaterialInfo.DerivedTextures.Num() == NumTextures);
 		for (int32 ti = 0; ti < NumTextures; ++ti)
 		{
 			const FTextureInfo& SourceTex = SourceMaterialInfo.SourceTextures[ti];
-			bool bConvertToSRGB = SourceTex.SourceTexture->SRGB;
-			const FTextureInfo& DerivedTex = DerivedMaterialInfo.SourceTextures[ti];
+			bool bConvertToSRGB = SourceTex.Texture->SRGB;
+			const FTextureInfo& DerivedTex = DerivedMaterialInfo.DerivedTextures[ti];
 			if (DerivedTex.bShouldBakeTexture)
 			{
 				FTexture2DBuilder TextureBuilder;
@@ -439,7 +538,7 @@ void UGenerateStaticMeshLODProcess::GetDerivedMaterialsPreview(FPreviewMaterials
 				UTexture2D* PreviewTex = TextureBuilder.GetTexture2D();
 				if (ensure(PreviewTex))
 				{
-					SourceToPreviewTexMap.Add(SourceTex.SourceTexture, PreviewTex);
+					SourceToPreviewTexMap.Add(SourceTex.Texture, PreviewTex);
 					MaterialSetOut.Textures.Add(PreviewTex);
 				}
 			}
@@ -457,15 +556,23 @@ void UGenerateStaticMeshLODProcess::GetDerivedMaterialsPreview(FPreviewMaterials
 	// create derived MIDs and point to new textures
 	for (int32 mi = 0; mi < NumMaterials; ++mi)
 	{
-		const FMaterialInfo& SourceMaterialInfo = SourceMaterials[mi];
-
+		const FSourceMaterialInfo& SourceMaterialInfo = SourceMaterials[mi];
 		UMaterialInterface* MaterialInterface = SourceMaterialInfo.SourceMaterial.MaterialInterface;
-		UMaterialInstanceDynamic* GeneratedMID = UMaterialInstanceDynamic::Create(MaterialInterface, NULL);
 
-		// rewrite texture parameters to new textures
-		UpdateMaterialTextureParameters(GeneratedMID, SourceMaterialInfo, SourceToPreviewTexMap, PreviewNormalMapTex);
+		if (SourceMaterialInfo.bIsReusable || SourceMaterialInfo.bIsPreviouslyGeneratedMaterial)
+		{
+			MaterialSetOut.Materials.Add(MaterialInterface);
+		}
+		else
+		{
+			// TODO: we should cache these instead of re-creating every time??
+			UMaterialInstanceDynamic* GeneratedMID = UMaterialInstanceDynamic::Create(MaterialInterface, NULL);
 
-		MaterialSetOut.Materials.Add(GeneratedMID);
+			// rewrite texture parameters to new textures
+			UpdateMaterialTextureParameters(GeneratedMID, SourceMaterialInfo, SourceToPreviewTexMap, PreviewNormalMapTex);
+
+			MaterialSetOut.Materials.Add(GeneratedMID);
+		}
 	}
 
 }
@@ -473,7 +580,7 @@ void UGenerateStaticMeshLODProcess::GetDerivedMaterialsPreview(FPreviewMaterials
 
 void UGenerateStaticMeshLODProcess::UpdateMaterialTextureParameters(
 	UMaterialInstanceDynamic* Material, 
-	const FMaterialInfo& SourceMaterialInfo,
+	const FSourceMaterialInfo& SourceMaterialInfo,
 	const TMap<UTexture2D*, UTexture2D*>& PreviewTextures, 
 	UTexture2D* PreviewNormalMap)
 {
@@ -493,7 +600,7 @@ void UGenerateStaticMeshLODProcess::UpdateMaterialTextureParameters(
 		}
 		else if (SourceTex.bShouldBakeTexture)
 		{
-			UTexture2D*const* FoundTexture = PreviewTextures.Find(SourceTex.SourceTexture);
+			UTexture2D*const* FoundTexture = PreviewTextures.Find(SourceTex.Texture);
 			if (ensure(FoundTexture))
 			{
 				FMaterialParameterInfo ParamInfo(SourceTex.ParameterName);
@@ -554,21 +661,26 @@ void UGenerateStaticMeshLODProcess::WriteDerivedTextures()
 	check(DerivedMaterials.Num() == NumMaterials);
 	for (int32 mi = 0; mi < NumMaterials; ++mi)
 	{
-		const FMaterialInfo& SourceMaterialInfo = SourceMaterials[mi];
-		FMaterialInfo& DerivedMaterialInfo = DerivedMaterials[mi];
+		const FSourceMaterialInfo& SourceMaterialInfo = SourceMaterials[mi];
+		if (SourceMaterialInfo.bIsReusable || SourceMaterialInfo.bIsPreviouslyGeneratedMaterial)
+		{
+			continue;
+		}
+
+		FDerivedMaterialInfo& DerivedMaterialInfo = DerivedMaterials[mi];
 
 		int32 NumTextures = SourceMaterialInfo.SourceTextures.Num();
-		check(DerivedMaterialInfo.SourceTextures.Num() == NumTextures);
+		check(DerivedMaterialInfo.DerivedTextures.Num() == NumTextures);
 		for (int32 ti = 0; ti < NumTextures; ++ti)
 		{
 			const FTextureInfo& SourceTex = SourceMaterialInfo.SourceTextures[ti];
-			if (SourceTexturesWritten.Contains(SourceTex.SourceTexture))
+			if (SourceTexturesWritten.Contains(SourceTex.Texture))
 			{
 				continue;
 			}
 
-			bool bConvertToSRGB = SourceTex.SourceTexture->SRGB;
-			FTextureInfo& DerivedTex = DerivedMaterialInfo.SourceTextures[ti];
+			bool bConvertToSRGB = SourceTex.Texture->SRGB;
+			FTextureInfo& DerivedTex = DerivedMaterialInfo.DerivedTextures[ti];
 
 			if (DerivedTex.bShouldBakeTexture == false)
 			{
@@ -581,18 +693,18 @@ void UGenerateStaticMeshLODProcess::WriteDerivedTextures()
 			TextureBuilder.Copy(DerivedTex.Image, bConvertToSRGB);
 			TextureBuilder.Commit(false);
 
-			DerivedTex.SourceTexture = TextureBuilder.GetTexture2D();
-			if (ensure(DerivedTex.SourceTexture))
+			DerivedTex.Texture = TextureBuilder.GetTexture2D();
+			if (ensure(DerivedTex.Texture))
 			{
-				AllDerivedTextures.Add(DerivedTex.SourceTexture);
+				AllDerivedTextures.Add(DerivedTex.Texture);
 
-				FTexture2DBuilder::CopyPlatformDataToSourceData(DerivedTex.SourceTexture, FTexture2DBuilder::ETextureType::Color);
+				FTexture2DBuilder::CopyPlatformDataToSourceData(DerivedTex.Texture, FTexture2DBuilder::ETextureType::Color);
 
 				// write asset
-				bool bWriteOK = WriteDerivedTexture(SourceTex.SourceTexture, DerivedTex.SourceTexture);
+				bool bWriteOK = WriteDerivedTexture(SourceTex.Texture, DerivedTex.Texture);
 				ensure(bWriteOK);
 
-				SourceTexturesWritten.Add(SourceTex.SourceTexture);
+				SourceTexturesWritten.Add(SourceTex.Texture);
 			}
 		}
 
@@ -678,13 +790,18 @@ void UGenerateStaticMeshLODProcess::WriteDerivedMaterials()
 	check(DerivedMaterials.Num() == NumMaterials);
 	for (int32 mi = 0; mi < NumMaterials; ++mi)
 	{
-		const FMaterialInfo& SourceMaterialInfo = SourceMaterials[mi];
-		FMaterialInfo& DerivedMaterialInfo = DerivedMaterials[mi];
+		const FSourceMaterialInfo& SourceMaterialInfo = SourceMaterials[mi];
+		if (SourceMaterialInfo.bIsReusable || SourceMaterialInfo.bIsPreviouslyGeneratedMaterial)
+		{
+			continue;
+		}
+
+		FDerivedMaterialInfo& DerivedMaterialInfo = DerivedMaterials[mi];
 
 		UMaterialInterface* MaterialInterface = SourceMaterialInfo.SourceMaterial.MaterialInterface;
 		if (MaterialInterface == nullptr)
 		{
-			DerivedMaterialInfo.SourceMaterial.MaterialInterface = nullptr;
+			DerivedMaterialInfo.DerivedMaterial.MaterialInterface = nullptr;
 			continue;
 		}
 		bool bSourceIsMIC = (Cast<UMaterialInstanceConstant>(MaterialInterface) != nullptr);
@@ -730,20 +847,22 @@ void UGenerateStaticMeshLODProcess::WriteDerivedMaterials()
 		UpdateMaterialTextureParameters(GeneratedMIC, DerivedMaterialInfo);
 
 		// update StaticMaterial
-		DerivedMaterialInfo.SourceMaterial.MaterialInterface = GeneratedMIC;
+		DerivedMaterialInfo.DerivedMaterial.MaterialInterface = GeneratedMIC;
+		DerivedMaterialInfo.DerivedMaterial.MaterialSlotName = FName(FString::Printf(TEXT("GeneratedMat%d"), mi));
+		DerivedMaterialInfo.DerivedMaterial.ImportedMaterialSlotName = DerivedMaterialInfo.DerivedMaterial.MaterialSlotName;
 	}
 }
 
 
 
-void UGenerateStaticMeshLODProcess::UpdateMaterialTextureParameters(UMaterialInstanceConstant* Material, FMaterialInfo& DerivedMaterialInfo)
+void UGenerateStaticMeshLODProcess::UpdateMaterialTextureParameters(UMaterialInstanceConstant* Material, FDerivedMaterialInfo& DerivedMaterialInfo)
 {
 	Material->Modify();
 
-	int32 NumTextures = DerivedMaterialInfo.SourceTextures.Num();
+	int32 NumTextures = DerivedMaterialInfo.DerivedTextures.Num();
 	for (int32 ti = 0; ti < NumTextures; ++ti)
 	{
-		FTextureInfo& DerivedTex = DerivedMaterialInfo.SourceTextures[ti];
+		FTextureInfo& DerivedTex = DerivedMaterialInfo.DerivedTextures[ti];
 
 		if (DerivedTex.bIsNormalMap)
 		{
@@ -755,7 +874,7 @@ void UGenerateStaticMeshLODProcess::UpdateMaterialTextureParameters(UMaterialIns
 		}
 		else if (DerivedTex.bShouldBakeTexture)
 		{
-			UTexture2D* NewTexture = DerivedTex.SourceTexture;
+			UTexture2D* NewTexture = DerivedTex.Texture;
 			if (ensure(NewTexture))
 			{
 				FMaterialParameterInfo ParamInfo(DerivedTex.ParameterName);
@@ -798,16 +917,28 @@ void UGenerateStaticMeshLODProcess::WriteDerivedStaticMeshAsset()
 	Converter.Convert(&DerivedLODMesh, *MeshDescription);
 	GeneratedStaticMesh->CommitMeshDescription(0);
 
+	// construct new material slots list
 	TArray<FStaticMaterial> NewMaterials;
-	for (FMaterialInfo& DerivedMaterialInfo : DerivedMaterials)
+	int32 NumMaterials = SourceMaterials.Num();
+	for (int32 mi = 0; mi < NumMaterials; ++mi)
 	{
-		NewMaterials.Add(DerivedMaterialInfo.SourceMaterial);
+		if (SourceMaterials[mi].bIsPreviouslyGeneratedMaterial)
+		{
+			NewMaterials.Add(SourceMaterials[mi].SourceMaterial);		// TODO: should be skipping this material...
+		}
+		else if (SourceMaterials[mi].bIsReusable)
+		{
+			NewMaterials.Add(SourceMaterials[mi].SourceMaterial);
+		}
+		else 
+		{
+			NewMaterials.Add(DerivedMaterials[mi].DerivedMaterial);
+		}
 	}
 
 	// update materials on generated mesh
 	GeneratedStaticMesh->SetStaticMaterials(NewMaterials);
 
-
 	// collision
 	FPhysicsDataCollection NewCollisionGeo;
 	NewCollisionGeo.Geometry = DerivedCollision;
@@ -839,44 +970,120 @@ void UGenerateStaticMeshLODProcess::WriteDerivedStaticMeshAsset()
 
 
 
-
 void UGenerateStaticMeshLODProcess::UpdateSourceStaticMeshAsset(bool bSetNewHDSourceAsset)
 {
 	GEditor->BeginTransaction(LOCTEXT("UpdateExistingAssetMessage", "Added Generated LOD"));
 
-	SourceStaticMesh->Modify();
+	FStaticMeshSourceModel& SrcModel = SourceStaticMesh->GetSourceModel(0);
+	SourceStaticMesh->ModifyMeshDescription(0);
 
-	if (bSetNewHDSourceAsset)
+	// if we want to save the input high-poly asset as the HiResSource, do that here
+	if (bSetNewHDSourceAsset && bUsingHiResSource == false)
 	{
+		SourceStaticMesh->ModifyHiResMeshDescription();
+
 		FMeshDescription* NewSourceMD = SourceStaticMesh->CreateHiResMeshDescription();
-		const FMeshDescription* ExistingSourceMD = SourceStaticMesh->GetMeshDescription(0);
-		*NewSourceMD = *ExistingSourceMD;
+		*NewSourceMD = *SourceMeshDescription;		// todo: can MoveTemp here, we don't need this memory anymore??
+
+		FStaticMeshSourceModel& HiResSrcModel = SourceStaticMesh->GetHiResSourceModel();
+		// Generally copy LOD0 build settings, although many of these will be ignored
+		HiResSrcModel.BuildSettings = SrcModel.BuildSettings;
+		// on the HiRes we store the existing normals and tangents, which we already auto-computed if necessary
+		HiResSrcModel.BuildSettings.bRecomputeNormals = false;
+		HiResSrcModel.BuildSettings.bRecomputeTangents = false;
+		// TODO: what should we do about Lightmap UVs?
+
 		SourceStaticMesh->CommitHiResMeshDescription();
 	}
 
-	TArray<FStaticMaterial> ExistingMaterials = SourceStaticMesh->GetStaticMaterials();
+	// Next bit is tricky, we have to build the final StaticMesh Material Set.
+	// We have the existing Source materials we want to keep, except if some
+	// were identified as being auto-generated by a previous run of AutoLOD, we want to leave those out.
+	// Then we want to add any New generated materials. 
+	// The main complication is that we cannot change the slot indices for the existing Source materials,
+	// as we would have to fix up the HighRes source. Ideally they are the first N slots, and we
+	// just append the new ones. But we cannot guarantee this, so if there are gaps we will interleave
+	// the new materials when possible.
+
+
+	// Figure out which derived materials we are actually going to use.
+	// We will not use materials we can re-use from the source, or materials that were previously auto-generated.
+	int32 NumMaterials = SourceMaterials.Num();
+	TArray<int32> DerivedMatSlotIndexMap;		// this maps from current DerivedMesh slot indices to their final slot indices
+	DerivedMatSlotIndexMap.SetNum(NumMaterials);
+	TArray<int32> DerivedMaterialsToAdd;		// list of derived material indices we need to store in the final material set
+	int32 NewMaterialCount = 0;
+	for (int32 mi = 0; mi < NumMaterials; ++mi)
+	{
+		if (SourceMaterials[mi].bIsPreviouslyGeneratedMaterial)
+		{
+			DerivedMatSlotIndexMap[mi] = -2;		// these materials do not appear in DerivedMesh or SourceMesh and should be skipped/discarded (todo: and deleted?)
+		}
+		else  if (SourceMaterials[mi].bIsReusable)	// if we can re-use existing material we just rewrite to existing material slot index
+		{
+			DerivedMatSlotIndexMap[mi] = mi;
+		}
+		else
+		{
+			DerivedMatSlotIndexMap[mi] = -1;		// will need to allocate a new slot for this material
+			DerivedMaterialsToAdd.Add(mi);
+			NewMaterialCount++;
+		}
+	}
+	int32 CurRemainingDerivedIdx = 0;
+
+	// Copy existing materials we want to keep to new StaticMesh materials set. 
+	// If there are any gaps left by skipping previously-derived materials, try to 
+	// tuck in a new derived material that is waiting to be allocated to a slot
+	TArray<FStaticMaterial> NewMaterialSet;
+	TArray<int32> DerivedMaterialSlotIndices;
+	for (int32 k = 0; k < SourceMaterials.Num(); ++k)
+	{
+		if (SourceMaterials[k].bIsPreviouslyGeneratedMaterial == false)
+		{
+			NewMaterialSet.Add(SourceMaterials[k].SourceMaterial);
+		}
+		else if (CurRemainingDerivedIdx < DerivedMaterialsToAdd.Num() )
+		{
+			int32 DerivedIdx = DerivedMaterialsToAdd[CurRemainingDerivedIdx];
+			CurRemainingDerivedIdx++;
+			DerivedMatSlotIndexMap[DerivedIdx] = NewMaterialSet.Num();
+			DerivedMaterialSlotIndices.Add(DerivedMatSlotIndexMap[DerivedIdx]);
+			NewMaterialSet.Add(DerivedMaterials[DerivedIdx].DerivedMaterial);
+		}
+		else
+		{
+			// we ran out of new materials to allocate and so just add empty ones??
+			ensure(false);
+			NewMaterialSet.Add(FStaticMaterial());
+		}
+	}
 
-	// append new materials to material set
-	TArray<int32> NewMatIndexMap;
-	for (FMaterialInfo& DerivedMaterialInfo : DerivedMaterials)
+	// if we have any new derived materials left, append them to the material set
+	while (CurRemainingDerivedIdx < DerivedMaterialsToAdd.Num())
 	{
-		int32 NewIndex = ExistingMaterials.Num();
-		ExistingMaterials.Add(DerivedMaterialInfo.SourceMaterial);
-		NewMatIndexMap.Add(NewIndex);
+		int32 DerivedIdx = DerivedMaterialsToAdd[CurRemainingDerivedIdx];
+		CurRemainingDerivedIdx++;
+		DerivedMatSlotIndexMap[DerivedIdx] = NewMaterialSet.Num();
+		DerivedMaterialSlotIndices.Add(DerivedMatSlotIndexMap[DerivedIdx]);
+		NewMaterialSet.Add(DerivedMaterials[DerivedIdx].DerivedMaterial);
 	}
 
-	// rewrite material IDs on derived mesh
+	// apply the material slot index rewrite map to the DerivedMesh
 	DerivedLODMesh.Attributes()->EnableMaterialID();
 	FDynamicMeshMaterialAttribute* MaterialIDs = DerivedLODMesh.Attributes()->GetMaterialID();
 	for (int32 tid : DerivedLODMesh.TriangleIndicesItr())
 	{
 		int32 CurMaterialID = MaterialIDs->GetValue(tid);
-		int32 NewMaterialID = NewMatIndexMap[CurMaterialID];
-		MaterialIDs->SetValue(tid, NewMaterialID);
+		int32 NewMaterialID = DerivedMatSlotIndexMap[CurMaterialID];
+		if (ensure(NewMaterialID >= 0))
+		{
+			MaterialIDs->SetValue(tid, NewMaterialID);
+		}
 	}
 
 	// update materials on generated mesh
-	SourceStaticMesh->SetStaticMaterials(ExistingMaterials);
+	SourceStaticMesh->SetStaticMaterials(NewMaterialSet);
 
 	// store new derived LOD as LOD 0
 	SourceStaticMesh->SetNumSourceModels(1);
@@ -888,14 +1095,32 @@ void UGenerateStaticMeshLODProcess::UpdateSourceStaticMeshAsset(bool bSetNewHDSo
 	FConversionToMeshDescriptionOptions ConversionOptions;
 	FDynamicMeshToMeshDescription Converter(ConversionOptions);
 	Converter.Convert(&DerivedLODMesh, *MeshDescription);
-	SourceStaticMesh->CommitMeshDescription(0);
+
+	// calculate tangents
+	Converter.UpdateTangents(&DerivedLODMesh, *MeshDescription, &DerivedLODMeshTangents);
+
+	// set slot names on MeshDescription to match those we set on the generated FStaticMaterials, 
+	// because StaticMesh RenderBuffer setup will do matching-name lookups and if it is NAME_None
+	// we will get the wrong material!
+	FStaticMeshAttributes Attributes(*MeshDescription);
+	TPolygonGroupAttributesRef<FName> PolygonGroupImportedMaterialSlotNames = Attributes.GetPolygonGroupMaterialSlotNames();
+	for (int32 SlotIdx : DerivedMaterialSlotIndices)
+	{
+		PolygonGroupImportedMaterialSlotNames.Set(SlotIdx, NewMaterialSet[SlotIdx].ImportedMaterialSlotName);
+	}
+
+	// Disable auto-generated normals/tangents, we need to use the ones we computed in LOD Generator
+	SrcModel.BuildSettings.bRecomputeNormals = false;
+	SrcModel.BuildSettings.bRecomputeTangents = false;
 
 	// this will prevent simplification?
-	FStaticMeshSourceModel& SrcModel = SourceStaticMesh->GetSourceModel(0);
 	SrcModel.ReductionSettings.MaxDeviation = 0.0f;
 	SrcModel.ReductionSettings.PercentTriangles = 1.0f;
 	SrcModel.ReductionSettings.PercentVertices = 1.0f;
 
+	// commit update
+	SourceStaticMesh->CommitMeshDescription(0);
+
 	// collision
 	FPhysicsDataCollection NewCollisionGeo;
 	NewCollisionGeo.Geometry = DerivedCollision;
diff --git a/Engine/Plugins/Experimental/MeshLODToolset/Source/MeshLODToolset/Public/Graphs/GenerateStaticMeshLODProcess.h b/Engine/Plugins/Experimental/MeshLODToolset/Source/MeshLODToolset/Public/Graphs/GenerateStaticMeshLODProcess.h
--- a/Engine/Plugins/Experimental/MeshLODToolset/Source/MeshLODToolset/Public/Graphs/GenerateStaticMeshLODProcess.h
+++ b/Engine/Plugins/Experimental/MeshLODToolset/Source/MeshLODToolset/Public/Graphs/GenerateStaticMeshLODProcess.h
@@ -15,6 +15,7 @@
 class UTexture2D;
 class UMaterialInstanceConstant;
 class UMaterialInstanceDynamic;
+struct FMeshDescription;
 
 
 UENUM()
@@ -216,11 +217,17 @@ protected:
 	FString SourceAssetFolder;
 	FString SourceAssetName;
 
+	// if true, we are building new LOD0 from the StaticMesh HiRes SourceModel, instead of from the mesh in LOD0
+	bool bUsingHiResSource = false;
+
+	// copy of input MeshDescription with autogenerated attributes computed
+	TSharedPtr<FMeshDescription> SourceMeshDescription;
+	// SourceMeshDescription converted to FDynamicMesh3
 	FDynamicMesh3 SourceMesh;
 
 	struct FTextureInfo
 	{
-		UTexture2D* SourceTexture;
+		UTexture2D* Texture = nullptr;
 		FName ParameterName;
 		FImageDimensions Dimensions;
 		TImageBuilder<FVector4f> Image;
@@ -229,38 +236,58 @@ protected:
 		bool bShouldBakeTexture = false;
 	};
 
-	struct FMaterialInfo
+	// Information about one of the input StaticMesh Materials. Computed in Initialize() and not modified afterwards
+	struct FSourceMaterialInfo
 	{
 		FStaticMaterial SourceMaterial;
 		TArray<FTextureInfo> SourceTextures;
-	};
-
-	TArray<FMaterialInfo> SourceMaterials;
 
+		bool bHasNormalMap = false;						// if true, Material has an exposed NormalMap input texture parameter
+		bool bHasTexturesToBake = false;				// if true, Material has at least one SourceTexture that should be baked
+		bool bIsReusable = false;						// if true, Material doesn't need any texture baking and can be re-used by LOD0
+		bool bIsPreviouslyGeneratedMaterial = false;	// if true, this Material was previously generated by AutoLOD and should be discarded. 
+														// Currently inferred from material being in LOD0 but not HiRes Source
+	};
 
+	// list of initial Source Materials, length equivalent to StaticMesh.StaticMaterials
+	TArray<FSourceMaterialInfo> SourceMaterials;
 
 	FString DerivedSuffix;
 	FString DerivedAssetFolder;
 	FString DerivedAssetPath;
 
-	FDynamicMesh3 DerivedLODMesh;
-	FMeshTangentsd DerivedLODMeshTangents;
-	FSimpleShapeSet3d DerivedCollision;
-	UE::GeometryFlow::FNormalMapImage DerivedNormalMapImage;
-	TArray<TUniquePtr<UE::GeometryFlow::FTextureImage>> DerivedTextureImages;
-	TMap<UTexture2D*, int32> SourceTextureToDerivedTexIndex;
+	FDynamicMesh3 DerivedLODMesh;				// the new generated LOD0 mesh
+	FMeshTangentsd DerivedLODMeshTangents;		// Tangents for DerivedLODMesh
+	FSimpleShapeSet3d DerivedCollision;			// Simple Collision for DerivedLODMesh
+
+	// Texture set potentially required by output Material set
+	UE::GeometryFlow::FNormalMapImage DerivedNormalMapImage;	// Normal Map
+	TArray<TUniquePtr<UE::GeometryFlow::FTextureImage>> DerivedTextureImages;	// generated Textures
+	TMap<UTexture2D*, int32> SourceTextureToDerivedTexIndex;	// mapping from input Textures to DerivedTextureImages index
+
+	// Information about an output Material
+	struct FDerivedMaterialInfo
+	{
+		int32 SourceMaterialIndex = -1;				// Index into SourceMaterials
+		bool bUseSourceMaterialDirectly = false;	// If true, do not create/use a Derived Material, directly re-use SourceMaterials[SourceMaterialIndex]
+
+		FStaticMaterial DerivedMaterial;			// points to generated Material
+		TArray<FTextureInfo> DerivedTextures;		// List of generated Textures
+	};
 
-	TArray<FMaterialInfo> DerivedMaterials;
+	// list of generated/Derived Materials. Length is the same as SourceMaterials and indices correspond, however some Derived Materials may not be initialized (if Source is reusable or generated)
+	TArray<FDerivedMaterialInfo> DerivedMaterials;
 
 	// This list is for accumulating derived UTexture2D's created during WriteDerivedTextures(). We have to
 	// maintain uproperty references to these or they may be garbage collected
 	UPROPERTY()
 	TSet<UTexture2D*> AllDerivedTextures;
 
+	// Derived Normal Map
 	UPROPERTY()
 	UTexture2D* DerivedNormalMapTex;
 
-	TUniquePtr<FGenerateMeshLODGraph> Generator;
+	TUniquePtr<FGenerateMeshLODGraph> Generator;			// active LODGenerator Graph
 	bool InitializeGenerator();
 
 	FGenerateStaticMeshLODProcessSettings CurrentSettings;
@@ -269,11 +296,11 @@ protected:
 	bool WriteDerivedTexture(UTexture2D* DerivedTexture, FString BaseTexName);
 	void WriteDerivedTextures();
 	void WriteDerivedMaterials();
-	void UpdateMaterialTextureParameters(UMaterialInstanceConstant* Material, FMaterialInfo& DerivedMaterialInfo);
+	void UpdateMaterialTextureParameters(UMaterialInstanceConstant* Material, FDerivedMaterialInfo& DerivedMaterialInfo);
 	void WriteDerivedStaticMeshAsset();
 
 	void UpdateSourceStaticMeshAsset(bool bSetNewHDSourceAsset);
 
-	void UpdateMaterialTextureParameters(UMaterialInstanceDynamic* Material, const FMaterialInfo& SourceMaterialInfo, 
+	void UpdateMaterialTextureParameters(UMaterialInstanceDynamic* Material, const FSourceMaterialInfo& SourceMaterialInfo,
 		const TMap<UTexture2D*,UTexture2D*>& PreviewTextures, UTexture2D* PreviewNormalMap);
 };
