ModelingTools: Add LODManagerTool, supports viewing different LODs in-viewport, including HiRes LOD. Action buttons to clear/reset HiRes LOD, and to remove unused materials (these are used in AutoLOD tool workflows). Minor API additions to UPreviewMesh and UPreviewGeometry. Move AutoLOD and this new Tool to a separate section of Modeling Mode toolbar
#rnx

diff --git a/Engine/Plugins/Experimental/MeshLODToolset/Source/MeshLODToolset/Private/Tools/LODManagerTool.cpp b/Engine/Plugins/Experimental/MeshLODToolset/Source/MeshLODToolset/Private/Tools/LODManagerTool.cpp
new file mode 100644
--- /dev/null
+++ b/Engine/Plugins/Experimental/MeshLODToolset/Source/MeshLODToolset/Private/Tools/LODManagerTool.cpp
@@ -0,0 +1,646 @@
+// Copyright Epic Games, Inc. All Rights Reserved.
+
+#include "Tools/LODManagerTool.h"
+#include "InteractiveToolManager.h"
+#include "ToolBuilderUtil.h"
+
+#include "ToolSetupUtil.h"
+#include "MathUtil.h"
+#include "AssetUtils/MeshDescriptionUtil.h"
+#include "MeshDescriptionToDynamicMesh.h"
+#include "IndexedMeshToDynamicMesh.h"
+
+#include "MeshDescription.h"
+#include "StaticMeshAttributes.h"
+#include "StaticMeshOperations.h"
+
+// for lightmap access
+#include "Engine/StaticMesh.h"
+#include "Components/StaticMeshComponent.h"
+
+#include "Components/PrimitiveComponent.h"
+
+
+#define LOCTEXT_NAMESPACE "ULODManagerTool"
+
+
+/*
+ * ToolBuilder
+ */
+
+
+bool ULODManagerToolBuilder::CanBuildTool(const FToolBuilderState& SceneState) const
+{
+	int32 NumStaticMeshes = ToolBuilderUtil::CountComponents(SceneState, [&](UActorComponent* Comp) { return Cast<UStaticMeshComponent>(Comp) != nullptr; });
+	int32 NumComponentTargets = ToolBuilderUtil::CountComponents(SceneState, CanMakeComponentTarget);
+	return (NumStaticMeshes > 0 && NumStaticMeshes == NumComponentTargets);
+}
+
+UInteractiveTool* ULODManagerToolBuilder::BuildTool(const FToolBuilderState& SceneState) const
+{
+	ULODManagerTool* NewTool = NewObject<ULODManagerTool>(SceneState.ToolManager);
+
+	TArray<UActorComponent*> ValidComponents = ToolBuilderUtil::FindAllComponents(SceneState,
+		[&](UActorComponent* Comp) { return Cast<UStaticMeshComponent>(Comp) != nullptr; });
+	check(ValidComponents.Num() > 0);
+
+	TArray<TUniquePtr<FPrimitiveComponentTarget>> ComponentTargets;
+	for (UActorComponent* ActorComponent : ValidComponents)
+	{
+		UStaticMeshComponent* MeshComponent = Cast<UStaticMeshComponent>(ActorComponent);
+		if ( MeshComponent )
+		{
+			ComponentTargets.Add(MakeComponentTarget(MeshComponent));
+		}
+	}
+
+	NewTool->SetSelection(MoveTemp(ComponentTargets));
+	NewTool->SetWorld(SceneState.World);
+
+	return NewTool;
+}
+
+
+
+
+void ULODManagerActionPropertySet::PostAction(ELODManagerToolActions Action)
+{
+	if (ParentTool.IsValid() && Cast<ULODManagerTool>(ParentTool))
+	{
+		Cast<ULODManagerTool>(ParentTool)->RequestAction(Action);
+	}
+}
+
+
+
+
+ULODManagerTool::ULODManagerTool()
+{
+}
+
+void ULODManagerTool::SetWorld(UWorld* World)
+{
+	this->TargetWorld = World;
+}
+
+void ULODManagerTool::Setup()
+{
+	UInteractiveTool::Setup();
+
+	HiResSourceModelActions = NewObject<ULODManagerHiResSourceModelActions>(this);
+	HiResSourceModelActions->Initialize(this);
+	AddToolPropertySource(HiResSourceModelActions);
+
+	MaterialActions = NewObject<ULODManagerMaterialActions>(this);
+	MaterialActions->Initialize(this);
+	AddToolPropertySource(MaterialActions);
+
+	if (ComponentTargets.Num() == 1)
+	{
+		LODInfoProperties = NewObject<ULODManagerLODProperties>(this);
+		AddToolPropertySource(LODInfoProperties);
+
+		LODPreviewProperties = NewObject<ULODManagerPreviewLODProperties>(this);
+		AddToolPropertySource(LODPreviewProperties);
+		LODPreviewProperties->VisibleLOD = this->DefaultLODName;
+		LODPreviewProperties->WatchProperty(LODPreviewProperties->VisibleLOD, [this](FString NewLOD) { bPreviewLODValid = false; });
+
+		LODPreview = NewObject<UPreviewMesh>(this);
+		LODPreview->CreateInWorld(TargetWorld, ComponentTargets[0]->GetWorldTransform());
+		LODPreview->SetVisible(false);
+
+		LODPreviewLines = NewObject<UPreviewGeometry>(this);
+		LODPreviewLines->CreateInWorld(TargetWorld, ComponentTargets[0]->GetWorldTransform());
+
+		FComponentMaterialSet MaterialSet;
+		ComponentTargets[0]->GetMaterialSet(MaterialSet);
+		LODPreview->SetMaterials(MaterialSet.Materials);
+
+		bLODInfoValid = false;
+	}
+
+	SetToolDisplayName(LOCTEXT("ToolName", "Manage LODs"));
+	GetToolManager()->DisplayMessage(
+		LOCTEXT("OnStartTool", "Inspect and Modify LODs of a StaticMesh Asset"),
+		EToolMessageLevel::UserNotification);
+}
+
+
+
+
+void ULODManagerTool::Shutdown(EToolShutdownType ShutdownType)
+{
+	if (LODPreview)
+	{
+		LODPreview->Disconnect();
+	}
+	for (TUniquePtr<FPrimitiveComponentTarget>& Target : ComponentTargets)
+	{
+		Target->SetOwnerVisibility(true);
+	}
+}
+
+
+
+void ULODManagerTool::RequestAction(ELODManagerToolActions ActionType)
+{
+	if (PendingAction == ELODManagerToolActions::NoAction)
+	{
+		PendingAction = ActionType;
+	}
+}
+
+
+
+void ULODManagerTool::OnTick(float DeltaTime)
+{
+	switch (PendingAction)
+	{
+	case ELODManagerToolActions::DeleteHiResSourceModel:
+		DeleteHiResSourceModel();
+		break;
+
+	case ELODManagerToolActions::MoveHiResToLOD0:
+		MoveHiResToLOD0();
+		break;
+
+	case ELODManagerToolActions::RemoveUnreferencedMaterials:
+		RemoveUnreferencedMaterials();
+		break;
+	}
+	PendingAction = ELODManagerToolActions::NoAction;
+
+	if (bLODInfoValid == false)
+	{
+		UpdateLODInfo();
+	}
+
+	if (bPreviewLODValid == false)
+	{
+		UpdatePreviewLOD();
+	}
+}
+
+
+
+UStaticMesh* ULODManagerTool::GetSingleStaticMesh()
+{
+	if (ComponentTargets.Num() > 1)
+	{
+		return nullptr;
+	}
+	TUniquePtr<FPrimitiveComponentTarget>& Target = ComponentTargets[0];
+	UStaticMeshComponent* StaticMeshComponent = Cast<UStaticMeshComponent>(Target->GetOwnerComponent());
+	if (StaticMeshComponent == nullptr || StaticMeshComponent->GetStaticMesh() == nullptr)
+	{
+		return nullptr;
+	}
+	UStaticMesh* StaticMesh = StaticMeshComponent->GetStaticMesh();
+	return StaticMesh;
+}
+
+
+
+
+
+
+
+void ULODManagerTool::UpdateLODInfo()
+{
+	UStaticMesh* StaticMesh = GetSingleStaticMesh();
+	if (bLODInfoValid == true || StaticMesh == nullptr)
+	{
+		return;
+	}
+	bLODInfoValid = true;
+
+	LODInfoProperties->bNaniteEnabled = StaticMesh->NaniteSettings.bEnabled;
+	LODInfoProperties->PercentTriangles = StaticMesh->NaniteSettings.PercentTriangles;
+
+	TArray<FStaticMaterial> CurMaterialSet = StaticMesh->GetStaticMaterials();
+	LODInfoProperties->Materials.Reset();
+	for (const FStaticMaterial& Material : StaticMesh->GetStaticMaterials())
+	{
+		LODInfoProperties->Materials.Add(Material);
+	}
+
+
+	// per-LOD info
+	LODInfoProperties->SourceLODs.Reset();
+	int32 NumSourceModels = StaticMesh->GetNumSourceModels();
+	for (int32 si = 0; si < NumSourceModels; ++si)
+	{
+		const FStaticMeshSourceModel& LODSourceModel = StaticMesh->GetSourceModel(si);
+		const FMeshDescription* LODMesh = StaticMesh->GetMeshDescription(si);
+		if (LODMesh != nullptr)		// generated LODs do not have mesh description...is that the only way to tell?
+		{
+			FLODManagerLODInfo LODInfo = { LODMesh->Vertices().Num(), LODMesh->Triangles().Num() };
+			LODInfoProperties->SourceLODs.Add(LODInfo);
+		}
+	}
+
+	LODInfoProperties->HiResSource.Reset();
+	if (StaticMesh->IsHiResMeshDescriptionValid())
+	{
+		const FMeshDescription* HiResMesh = StaticMesh->GetHiResMeshDescription();
+		FLODManagerLODInfo LODInfo = { HiResMesh->Vertices().Num(), HiResMesh->Triangles().Num() };
+		LODInfoProperties->HiResSource.Add(LODInfo);
+	}
+
+	LODInfoProperties->RenderLODs.Reset();
+	if (StaticMesh->HasValidRenderData())
+	{
+		const FStaticMeshRenderData* RenderData = StaticMesh->GetRenderData();
+		for (const FStaticMeshLODResources& LODResource : RenderData->LODResources)
+		{
+			FLODManagerLODInfo LODInfo;
+			LODInfo.VertexCount = LODResource.VertexBuffers.PositionVertexBuffer.GetNumVertices();
+			LODInfo.TriangleCount = LODResource.IndexBuffer.GetNumIndices() / 3;
+			LODInfoProperties->RenderLODs.Add(LODInfo);
+		}
+	}
+
+
+	UpdateLODNames();
+}
+
+
+
+
+void ULODManagerTool::UpdateLODNames()
+{
+	UStaticMesh* StaticMesh = GetSingleStaticMesh();
+	if (StaticMesh == nullptr || LODPreviewProperties == nullptr)
+	{
+		return;
+	}
+
+	TArray<FString>& UILODNames = LODPreviewProperties->LODNamesList;
+	UILODNames.Reset();
+	ActiveLODNames.Reset();
+
+	UILODNames.Add(DefaultLODName);
+	ActiveLODNames.Add(DefaultLODName, FLODName());
+
+	if (StaticMesh->IsHiResMeshDescriptionValid())
+	{
+		FString LODName(TEXT("SourceModel HiRes"));
+		UILODNames.Add(LODName);
+		ActiveLODNames.Add(LODName, { -1, -1, 0 });
+	}
+
+	int32 NumSourceModels = StaticMesh->GetNumSourceModels();
+	for (int32 si = 0; si < NumSourceModels; ++si)
+	{
+		const FStaticMeshSourceModel& LODSourceModel = StaticMesh->GetSourceModel(si);
+		const FMeshDescription* LODMesh = StaticMesh->GetMeshDescription(si);
+		if (LODMesh != nullptr)		// generated LODs do not have mesh description...is that the only way to tell?
+		{
+			FString LODName = FString::Printf(TEXT("SourceModel LOD%d"), si);
+			UILODNames.Add(LODName);
+			ActiveLODNames.Add(LODName, { si, -1, -1 });
+		}
+	}
+
+	if (StaticMesh->HasValidRenderData())
+	{
+		const FStaticMeshRenderData* RenderData = StaticMesh->GetRenderData();
+		int32 Index = 0;
+		for (const FStaticMeshLODResources& LODResource : RenderData->LODResources)
+		{
+			FString LODName = FString::Printf(TEXT("RenderData LOD%d"), Index);
+			UILODNames.Add(LODName);
+			ActiveLODNames.Add(LODName, { -1, Index, -1 });
+			Index++;
+		}
+	}
+
+
+	bPreviewLODValid = false;    // should not always be necessary...
+}
+
+
+
+
+
+void ULODManagerTool::UpdatePreviewLines(FLODMeshInfo& MeshInfo)
+{
+	if (MeshInfo.bInfoCached == false)
+	{
+		MeshInfo.bInfoCached = true;
+		for (int32 eid : MeshInfo.Mesh.EdgeIndicesItr())
+		{
+			if (MeshInfo.Mesh.IsBoundaryEdge(eid))
+			{
+				MeshInfo.BoundaryEdges.Add(eid);
+			}
+		}
+	}
+
+	float LineWidthMultiplier = 1.0f;
+	FColor BoundaryEdgeColor(240, 15, 15);
+	float BoundaryEdgeThickness = LineWidthMultiplier * 4.0f;
+	float BoundaryEdgeDepthBias = 2.0f;
+
+	ULineSetComponent* BoundaryLines = LODPreviewLines->FindLineSet(TEXT("BoundaryLines"));
+	if (BoundaryLines == nullptr)
+	{
+		BoundaryLines = LODPreviewLines->AddLineSet(TEXT("BoundaryLines"));
+	}
+	if (BoundaryLines)
+	{
+		for (int32 eid : MeshInfo.BoundaryEdges)
+		{
+			FVector3d A, B;
+			MeshInfo.Mesh.GetEdgeV(eid, A, B);
+			BoundaryLines->AddLine((FVector)A, (FVector)B, BoundaryEdgeColor, BoundaryEdgeThickness, BoundaryEdgeDepthBias);
+		}
+	}
+}
+
+
+
+void ULODManagerTool::UpdatePreviewLOD()
+{
+	UStaticMesh* StaticMesh = GetSingleStaticMesh();
+	if (StaticMesh == nullptr || LODPreviewProperties == nullptr || bPreviewLODValid == true)
+	{
+		return;
+	}
+	bPreviewLODValid = true;
+
+	FString SelectedLOD = LODPreviewProperties->VisibleLOD;
+	const FLODName* FoundName = ActiveLODNames.Find(SelectedLOD);
+	if (SelectedLOD.IsEmpty() || FoundName == nullptr || FoundName->IsDefault() )
+	{
+		// badness
+		LODPreview->SetVisible(false);
+		ComponentTargets[0]->SetOwnerVisibility(true);
+	}
+
+	TUniquePtr<FLODMeshInfo>* Found = LODMeshCache.Find(SelectedLOD);
+	if (Found == nullptr)
+	{
+		CacheLODMesh(SelectedLOD, *FoundName);
+		Found = LODMeshCache.Find(SelectedLOD);
+	}
+	
+	if (Found)
+	{
+		ComponentTargets[0]->SetOwnerVisibility(false);
+		LODPreview->ReplaceMesh( (*Found)->Mesh );
+		LODPreview->SetVisible(true);
+		UpdatePreviewLines(**Found);
+		LODPreviewLines->SetAllVisible(true);
+	}
+	else
+	{
+		LODPreview->SetVisible(false);
+		LODPreviewLines->SetAllVisible(false);
+		ComponentTargets[0]->SetOwnerVisibility(true);
+	}
+}
+
+
+
+
+bool ULODManagerTool::CacheLODMesh(const FString& Name, FLODName LODName)
+{
+	UStaticMesh* StaticMesh = GetSingleStaticMesh();
+	if (!ensure(StaticMesh)) return false;
+
+	FMeshDescriptionToDynamicMesh Converter;
+
+	if (LODName.OtherIndex == 0)		// HiRes SourceModel
+	{
+		FMeshDescription* HiResMeshDescription = StaticMesh->GetHiResMeshDescription();
+		if (HiResMeshDescription != nullptr)
+		{
+			TUniquePtr<FLODMeshInfo> NewMeshInfo = MakeUnique<FLODMeshInfo>();
+			Converter.Convert(HiResMeshDescription, NewMeshInfo->Mesh);
+			LODMeshCache.Add(Name, MoveTemp(NewMeshInfo));
+			return true;
+		}
+
+	}
+	else if (LODName.SourceModelIndex >= 0)
+	{
+		FMeshDescription* LODMeshDescription = StaticMesh->GetMeshDescription(LODName.SourceModelIndex);
+		if (LODMeshDescription != nullptr)
+		{
+			TUniquePtr<FLODMeshInfo> NewMeshInfo = MakeUnique<FLODMeshInfo>();
+			Converter.Convert(LODMeshDescription, NewMeshInfo->Mesh);
+			LODMeshCache.Add(Name, MoveTemp(NewMeshInfo));
+			return true;
+		}
+	}
+	else if (LODName.RenderDataIndex >= 0)
+	{
+		if (StaticMesh->HasValidRenderData())
+		{
+			const FStaticMeshRenderData* RenderData = StaticMesh->GetRenderData();
+			if (LODName.RenderDataIndex < RenderData->LODResources.Num())
+			{
+				const FStaticMeshLODResources& LODResource = RenderData->LODResources[LODName.RenderDataIndex];
+				TUniquePtr<FLODMeshInfo> NewMeshInfo = MakeUnique<FLODMeshInfo>();
+				UE::Conversion::RenderBuffersToDynamicMesh(LODResource.VertexBuffers, LODResource.IndexBuffer, LODResource.Sections, NewMeshInfo->Mesh);
+				LODMeshCache.Add(Name, MoveTemp(NewMeshInfo));
+				return true;
+			}
+		}
+	}
+
+	return false;
+}
+
+
+
+void ULODManagerTool::DeleteHiResSourceModel()
+{
+	GetToolManager()->BeginUndoTransaction(LOCTEXT("DeleteHiResSourceModel", "Delete HiRes Source"));
+	for (TUniquePtr<FPrimitiveComponentTarget>& Target : ComponentTargets)
+	{
+		UStaticMeshComponent* StaticMeshComponent = Cast<UStaticMeshComponent>(Target->GetOwnerComponent());
+		if (StaticMeshComponent == nullptr || StaticMeshComponent->GetStaticMesh() == nullptr)
+		{
+			continue;
+		}
+		UStaticMesh* StaticMesh = StaticMeshComponent->GetStaticMesh();
+
+		if (StaticMesh->IsHiResMeshDescriptionValid())
+		{
+			StaticMesh->Modify();
+
+			StaticMesh->ModifyHiResMeshDescription();
+			StaticMesh->ClearHiResMeshDescription();
+			StaticMesh->CommitHiResMeshDescription();
+
+			StaticMesh->PostEditChange();
+		}
+	}
+
+	GetToolManager()->EndUndoTransaction();
+
+	LODMeshCache.Reset();
+	bLODInfoValid = false;
+}
+
+
+
+void ULODManagerTool::MoveHiResToLOD0()
+{
+	GetToolManager()->BeginUndoTransaction(LOCTEXT("RestoreLOD0", "Move HiRes to LOD0"));
+	for (TUniquePtr<FPrimitiveComponentTarget>& Target : ComponentTargets)
+	{
+		UStaticMeshComponent* StaticMeshComponent = Cast<UStaticMeshComponent>(Target->GetOwnerComponent());
+		if (StaticMeshComponent == nullptr || StaticMeshComponent->GetStaticMesh() == nullptr)
+		{
+			continue;
+		}
+		UStaticMesh* StaticMesh = StaticMeshComponent->GetStaticMesh();
+
+		if (StaticMesh->IsHiResMeshDescriptionValid())
+		{
+			StaticMesh->Modify();
+
+			StaticMesh->ModifyHiResMeshDescription();
+			FMeshDescription* HiResMeshDescription = StaticMesh->GetHiResMeshDescription();
+
+			StaticMesh->ModifyMeshDescription(0);
+			FMeshDescription* LOD0MeshDescription = StaticMesh->GetMeshDescription(0);
+			*LOD0MeshDescription = *HiResMeshDescription;
+
+			StaticMesh->ClearHiResMeshDescription();
+			StaticMesh->CommitHiResMeshDescription();
+
+			StaticMesh->CommitMeshDescription(0);
+
+			StaticMesh->PostEditChange();
+		}
+	}
+
+	GetToolManager()->EndUndoTransaction();
+
+	LODMeshCache.Reset();
+	bLODInfoValid = false;
+}
+
+
+
+
+void ULODManagerTool::RemoveUnreferencedMaterials()
+{
+	GetToolManager()->BeginUndoTransaction(LOCTEXT("RemoveUnreferencedMaterials", "Remove Unreferenced Materials"));
+	for (TUniquePtr<FPrimitiveComponentTarget>& Target : ComponentTargets)
+	{
+		UStaticMeshComponent* StaticMeshComponent = Cast<UStaticMeshComponent>(Target->GetOwnerComponent());
+		if (StaticMeshComponent == nullptr || StaticMeshComponent->GetStaticMesh() == nullptr)
+		{
+			continue;
+		}
+		UStaticMesh* StaticMesh = StaticMeshComponent->GetStaticMesh();
+
+		TArray<FStaticMaterial> CurMaterialSet = StaticMesh->GetStaticMaterials();
+		int32 NumMaterials = CurMaterialSet.Num();
+		TArray<bool> MatUsedFlags;
+		MatUsedFlags.Init(false, NumMaterials);
+
+		TArray<const FMeshDescription*> Meshes;
+		for (int32 k = 0; k < StaticMesh->GetNumSourceModels(); ++k)
+		{
+			if (StaticMesh->IsSourceModelValid(k))
+			{
+				Meshes.Add(StaticMesh->GetMeshDescription(k));
+			}
+		}
+		if (StaticMesh->IsHiResMeshDescriptionValid())
+		{
+			Meshes.Add(StaticMesh->GetHiResMeshDescription());
+		}
+		int32 NumMeshes = Meshes.Num();
+
+
+		// for each mesh, collect list of material indices, and then set MatUsedFlags
+		TArray<TArray<int32>> MeshMaterialIndices;
+		MeshMaterialIndices.SetNum(NumMeshes);
+		for (int32 mi = 0; mi < NumMeshes; ++mi)
+		{
+			TMap<FPolygonGroupID, int32> PolygonGroupToSectionIndex;
+			TMap<FPolygonGroupID, int32> PolygonGroupToMaterialIndex;
+			FStaticMeshConstAttributes Attributes(*Meshes[mi]);
+			TPolygonGroupAttributesConstRef<FName> PolygonGroupImportedMaterialSlotNames = Attributes.GetPolygonGroupMaterialSlotNames();
+			for (FPolygonGroupID PolygonGroupID : Meshes[mi]->PolygonGroups().GetElementIDs() )
+			{
+				int32& SectionIndex = PolygonGroupToSectionIndex.FindOrAdd(PolygonGroupID);
+				int32 MaterialIndex = StaticMesh->GetMaterialIndexFromImportedMaterialSlotName(PolygonGroupImportedMaterialSlotNames[PolygonGroupID]);
+				if (MaterialIndex == INDEX_NONE)
+				{
+					MaterialIndex = PolygonGroupID.GetValue();
+				}
+				PolygonGroupToMaterialIndex.Add(PolygonGroupID, MaterialIndex);
+			}
+
+			for (const FTriangleID TriangleID : Meshes[mi]->Triangles().GetElementIDs())
+			{
+				const FPolygonGroupID PolygonGroupID = Meshes[mi]->GetTrianglePolygonGroup(TriangleID);
+				const int32 MaterialIndex = PolygonGroupToMaterialIndex[PolygonGroupID];
+				MatUsedFlags[MaterialIndex] = true;
+				MeshMaterialIndices[mi].AddUnique(MaterialIndex);
+			}
+		}
+
+		// construct new material set and ID remap
+		TArray<FStaticMaterial> NewMaterialSet;
+		TArray<int32> MaterialIDMap;
+		TArray<int32> RemovedMaterials;
+		TMap<FPolygonGroupID, FPolygonGroupID> RemapPolygonGroups;
+		MaterialIDMap.Init(-1, NumMaterials);
+		for (int32 k = 0; k < NumMaterials; ++k)
+		{
+			if (MatUsedFlags[k] == true )
+			{
+				MaterialIDMap[k] = NewMaterialSet.Num();
+				RemapPolygonGroups.Add(FPolygonGroupID(k), FPolygonGroupID(MaterialIDMap[k]));
+				NewMaterialSet.Add(CurMaterialSet[k]);
+			}
+			else
+			{
+				RemovedMaterials.Add(k);
+			}
+		}
+		if (NewMaterialSet.Num() == CurMaterialSet.Num())
+		{
+			continue;
+		}
+
+		StaticMesh->Modify();
+
+		// update material set
+		StaticMesh->GetStaticMaterials() = NewMaterialSet;
+
+		// now remap groups in each mesh
+		int32 MeshIndex = 0;
+		for (int32 k = 0; k < StaticMesh->GetNumSourceModels(); ++k)
+		{
+			if (StaticMesh->IsSourceModelValid(k))
+			{
+				StaticMesh->ModifyMeshDescription(k);
+				FMeshDescription* Mesh = StaticMesh->GetMeshDescription(k);
+				Mesh->RemapPolygonGroups(RemapPolygonGroups);
+				StaticMesh->CommitMeshDescription(k);
+			}
+		}
+		if (StaticMesh->IsHiResMeshDescriptionValid())
+		{
+			StaticMesh->ModifyHiResMeshDescription();
+			FMeshDescription* HiResMeshDescription = StaticMesh->GetHiResMeshDescription();
+			HiResMeshDescription->RemapPolygonGroups(RemapPolygonGroups);
+			StaticMesh->CommitHiResMeshDescription();
+		}
+
+		StaticMesh->PostEditChange();
+	}
+
+	GetToolManager()->EndUndoTransaction();
+	bLODInfoValid = false;
+}
diff --git a/Engine/Plugins/Experimental/MeshLODToolset/Source/MeshLODToolset/Public/Tools/LODManagerTool.h b/Engine/Plugins/Experimental/MeshLODToolset/Source/MeshLODToolset/Public/Tools/LODManagerTool.h
new file mode 100644
--- /dev/null
+++ b/Engine/Plugins/Experimental/MeshLODToolset/Source/MeshLODToolset/Public/Tools/LODManagerTool.h
@@ -0,0 +1,251 @@
+// Copyright Epic Games, Inc. All Rights Reserved.
+
+#pragma once
+
+#include "CoreMinimal.h"
+#include "MultiSelectionTool.h"
+#include "InteractiveToolBuilder.h"
+#include "Engine/StaticMesh.h"
+#include "PreviewMesh.h"
+#include "Drawing/PreviewGeometryActor.h"
+#include "LODManagerTool.generated.h"
+
+// predeclarations
+class IAssetGenerationAPI;
+class UMaterialInterface;
+
+/**
+ *
+ */
+UCLASS()
+class MESHLODTOOLSET_API ULODManagerToolBuilder : public UInteractiveToolBuilder
+{
+	GENERATED_BODY()
+
+public:
+	IAssetGenerationAPI* AssetAPI = nullptr;
+
+	virtual bool CanBuildTool(const FToolBuilderState& SceneState) const override;
+	virtual UInteractiveTool* BuildTool(const FToolBuilderState& SceneState) const override;
+};
+
+
+USTRUCT()
+struct MESHLODTOOLSET_API FLODManagerLODInfo
+{
+	GENERATED_BODY()
+
+	UPROPERTY(VisibleAnywhere, Category = LODInformation)
+	int32 VertexCount;
+
+	UPROPERTY(VisibleAnywhere, Category = LODInformation)
+	int32 TriangleCount;
+};
+
+
+
+UCLASS()
+class MESHLODTOOLSET_API ULODManagerLODProperties : public UInteractiveToolPropertySet
+{
+	GENERATED_BODY()
+public:
+	UPROPERTY(VisibleAnywhere, Category = LODInformation)
+	TArray<FLODManagerLODInfo> SourceLODs;
+
+	UPROPERTY(VisibleAnywhere, Category = LODInformation)
+	TArray<FLODManagerLODInfo> HiResSource;
+
+
+	UPROPERTY(VisibleAnywhere, Category = LODInformation)
+	TArray<FLODManagerLODInfo> RenderLODs;
+
+	UPROPERTY(VisibleAnywhere, Category = Nanite, meta = (DisplayName="Enabled"))
+	bool bNaniteEnabled = false;
+
+	UPROPERTY(VisibleAnywhere, Category = Nanite)
+	float PercentTriangles = 0.0f;
+
+	UPROPERTY(VisibleAnywhere, Category = Materials)
+	TArray<FStaticMaterial> Materials;
+
+};
+
+
+
+UCLASS()
+class MESHLODTOOLSET_API ULODManagerPreviewLODProperties : public UInteractiveToolPropertySet
+{
+	GENERATED_BODY()
+public:
+
+	UPROPERTY(EditAnywhere, Category = LODPreview, meta = (DisplayName = "Show LOD", GetOptions = GetLODNamesFunc))
+	FString VisibleLOD;
+
+	UFUNCTION()
+	const TArray<FString>& GetLODNamesFunc() const { return LODNamesList; }
+
+	UPROPERTY(meta = (TransientToolProperty))
+	TArray<FString> LODNamesList;
+};
+
+
+
+
+
+
+
+UENUM()
+enum class ELODManagerToolActions
+{
+	NoAction,
+
+	MoveHiResToLOD0,
+	DeleteHiResSourceModel,
+	RemoveUnreferencedMaterials
+};
+
+
+
+UCLASS()
+class MESHLODTOOLSET_API ULODManagerActionPropertySet : public UInteractiveToolPropertySet
+{
+	GENERATED_BODY()
+
+public:
+	TWeakObjectPtr<ULODManagerTool> ParentTool;
+
+	void Initialize(ULODManagerTool* ParentToolIn) { ParentTool = ParentToolIn; }
+	void PostAction(ELODManagerToolActions Action);
+
+};
+
+
+UCLASS()
+class MESHLODTOOLSET_API ULODManagerHiResSourceModelActions : public ULODManagerActionPropertySet
+{
+	GENERATED_BODY()
+public:
+	/** Move the HiRes Source Model to LOD0 */
+	UFUNCTION(CallInEditor, Category = HiResSourceModel, meta = (DisplayPriority = 0))
+	void MoveToLOD0()
+	{
+		PostAction(ELODManagerToolActions::MoveHiResToLOD0);
+	}
+
+	/** Delete the HiRes Source Model */
+	UFUNCTION(CallInEditor, Category = HiResSourceModel, meta = (DisplayPriority = 1))
+	void Delete()
+	{
+		PostAction(ELODManagerToolActions::DeleteHiResSourceModel);
+	}
+};
+
+
+UCLASS()
+class MESHLODTOOLSET_API ULODManagerMaterialActions : public ULODManagerActionPropertySet
+{
+	GENERATED_BODY()
+public:
+	/** Discard any Materials that are not referenced by any LOD */
+	UFUNCTION(CallInEditor, Category = MaterialSet, meta = (DisplayPriority = 0))
+	void CleanMaterials()
+	{
+		PostAction(ELODManagerToolActions::RemoveUnreferencedMaterials);
+	}
+};
+
+
+
+
+
+/**
+ * Mesh Attribute Editor Tool
+ */
+UCLASS()
+class MESHLODTOOLSET_API ULODManagerTool : public UMultiSelectionTool
+{
+	GENERATED_BODY()
+
+public:
+	ULODManagerTool();
+
+	virtual void Setup() override;
+	virtual void Shutdown(EToolShutdownType ShutdownType) override;
+	virtual void OnTick(float DeltaTime) override;
+
+	virtual void SetWorld(UWorld* World);
+
+	virtual bool HasCancel() const override { return false; }
+	virtual bool HasAccept() const override { return false; }
+	virtual bool CanAccept() const override { return false; }
+
+	virtual void RequestAction(ELODManagerToolActions ActionType);
+
+protected:
+
+	UPROPERTY()
+	ULODManagerLODProperties* LODInfoProperties;
+
+	UPROPERTY()
+	ULODManagerPreviewLODProperties* LODPreviewProperties;
+
+	UPROPERTY()
+	ULODManagerHiResSourceModelActions* HiResSourceModelActions;
+
+	UPROPERTY()
+	ULODManagerMaterialActions* MaterialActions;
+
+
+
+public:
+	UFUNCTION()
+	void DeleteHiResSourceModel();
+
+	UFUNCTION()
+	void MoveHiResToLOD0();
+
+	UFUNCTION()
+	void RemoveUnreferencedMaterials();
+
+protected:
+	UWorld* TargetWorld;
+
+	UStaticMesh* GetSingleStaticMesh();
+
+	ELODManagerToolActions PendingAction = ELODManagerToolActions::NoAction;
+	
+	bool bLODInfoValid = false;
+	void UpdateLODInfo();
+
+	struct FLODName
+	{
+		int32 SourceModelIndex = -1;
+		int32 RenderDataIndex = -1;
+		int32 OtherIndex = -1;
+		bool IsDefault() const { return SourceModelIndex == -1 && RenderDataIndex == -1 && OtherIndex == -1; }
+	};
+	FString DefaultLODName = FString("Default");
+	TMap<FString, FLODName> ActiveLODNames;
+	void UpdateLODNames();
+
+	struct FLODMeshInfo
+	{
+		FDynamicMesh3 Mesh;
+
+		bool bInfoCached = false;
+		TArray<int> BoundaryEdges;
+	};
+
+	TMap<FString, TUniquePtr<FLODMeshInfo>> LODMeshCache;
+
+	UPROPERTY()
+	UPreviewMesh* LODPreview;
+
+	UPROPERTY()
+	UPreviewGeometry* LODPreviewLines;
+
+	bool bPreviewLODValid = false;
+	void UpdatePreviewLOD();
+	void UpdatePreviewLines(FLODMeshInfo& LODMeshInfo);
+	bool CacheLODMesh(const FString& Name, FLODName LODName);
+};
diff --git a/Engine/Plugins/Experimental/MeshModelingToolset/Source/ModelingComponents/Private/Drawing/PreviewGeometryActor.cpp b/Engine/Plugins/Experimental/MeshModelingToolset/Source/ModelingComponents/Private/Drawing/PreviewGeometryActor.cpp
--- a/Engine/Plugins/Experimental/MeshModelingToolset/Source/ModelingComponents/Private/Drawing/PreviewGeometryActor.cpp
+++ b/Engine/Plugins/Experimental/MeshModelingToolset/Source/ModelingComponents/Private/Drawing/PreviewGeometryActor.cpp
@@ -47,6 +47,34 @@ void UPreviewGeometry::Disconnect()
 
 
 
+FTransform UPreviewGeometry::GetTransform() const
+{
+	if (ParentActor != nullptr)
+	{
+		return ParentActor->GetTransform();
+	}
+	return FTransform();
+}
+
+void UPreviewGeometry::SetTransform(const FTransform& UseTransform)
+{
+	if (ParentActor != nullptr)
+	{
+		ParentActor->SetActorTransform(UseTransform);
+	}
+}
+
+
+void UPreviewGeometry::SetAllVisible(bool bVisible)
+{
+	for (TPair<FString, ULineSetComponent*> Entry : LineSets)
+	{
+		Entry.Value->SetVisibility(bVisible);
+	}
+}
+
+
+
 
 ULineSetComponent* UPreviewGeometry::AddLineSet(const FString& SetIdentifier)
 {
diff --git a/Engine/Plugins/Experimental/MeshModelingToolset/Source/ModelingComponents/Private/PreviewMesh.cpp b/Engine/Plugins/Experimental/MeshModelingToolset/Source/ModelingComponents/Private/PreviewMesh.cpp
--- a/Engine/Plugins/Experimental/MeshModelingToolset/Source/ModelingComponents/Private/PreviewMesh.cpp
+++ b/Engine/Plugins/Experimental/MeshModelingToolset/Source/ModelingComponents/Private/PreviewMesh.cpp
@@ -373,6 +373,11 @@ void UPreviewMesh::InitializeMesh(FMeshDescription* MeshDescription)
 }
 
 
+void UPreviewMesh::ReplaceMesh(const FDynamicMesh3& NewMesh)
+{
+	ReplaceMesh(FDynamicMesh3(NewMesh));
+}
+
 void UPreviewMesh::ReplaceMesh(FDynamicMesh3&& NewMesh)
 {
 	FDynamicMesh3* Mesh = DynamicMeshComponent->GetMesh();
diff --git a/Engine/Plugins/Experimental/MeshModelingToolset/Source/ModelingComponents/Public/Drawing/PreviewGeometryActor.h b/Engine/Plugins/Experimental/MeshModelingToolset/Source/ModelingComponents/Public/Drawing/PreviewGeometryActor.h
--- a/Engine/Plugins/Experimental/MeshModelingToolset/Source/ModelingComponents/Public/Drawing/PreviewGeometryActor.h
+++ b/Engine/Plugins/Experimental/MeshModelingToolset/Source/ModelingComponents/Public/Drawing/PreviewGeometryActor.h
@@ -58,6 +58,24 @@ public:
 
 
 
+	/**
+	 * Get the current transform on the preview 
+	 */
+	FTransform GetTransform() const;
+
+	/**
+	 * Set the transform on the preview mesh
+	 */
+	void SetTransform(const FTransform& UseTransform);
+
+
+	/**
+	 * Set visibility state of the preview mesh
+	 */
+	void SetAllVisible(bool bVisible);
+
+
+
 	//
 	// Line Sets
 	//
@@ -169,5 +187,4 @@ protected:
 
 	// called at the beginning of Disconnect() to allow subclasses to perform additional cleanup
 	virtual void OnDisconnected() {}
-
 };
\ No newline at end of file
diff --git a/Engine/Plugins/Experimental/MeshModelingToolset/Source/ModelingComponents/Public/PreviewMesh.h b/Engine/Plugins/Experimental/MeshModelingToolset/Source/ModelingComponents/Public/PreviewMesh.h
--- a/Engine/Plugins/Experimental/MeshModelingToolset/Source/ModelingComponents/Public/PreviewMesh.h
+++ b/Engine/Plugins/Experimental/MeshModelingToolset/Source/ModelingComponents/Public/PreviewMesh.h
@@ -332,6 +332,11 @@ public:
 	// Edit access to internal mesh, and change-tracking/notification
 	// 
 
+	/**
+	 * Replace mesh with new mesh
+	 */
+	void ReplaceMesh(const FDynamicMesh3& NewMesh);
+
 	/**
 	 * Replace mesh with new mesh
 	 */
diff --git a/Engine/Plugins/Experimental/ModelingToolsEditorMode/Source/ModelingToolsEditorMode/Private/ModelingToolsEditorMode.cpp b/Engine/Plugins/Experimental/ModelingToolsEditorMode/Source/ModelingToolsEditorMode/Private/ModelingToolsEditorMode.cpp
--- a/Engine/Plugins/Experimental/ModelingToolsEditorMode/Source/ModelingToolsEditorMode/Private/ModelingToolsEditorMode.cpp
+++ b/Engine/Plugins/Experimental/ModelingToolsEditorMode/Source/ModelingToolsEditorMode/Private/ModelingToolsEditorMode.cpp
@@ -88,6 +88,7 @@
 
 // asset tools
 #include "Tools/GenerateStaticMeshLODAssetTool.h"
+#include "Tools/LODManagerTool.h"
 
 #include "EditorModeManager.h"
 
@@ -416,6 +417,11 @@ void UModelingToolsEditorMode::Enter()
 	DuplicateMeshesToolBuilder->bIsDuplicateTool = true;
 	RegisterTool(ToolManagerCommands.BeginDuplicateMeshesTool, TEXT("DuplicateMeshesTool"), DuplicateMeshesToolBuilder);
 
+
+	ULODManagerToolBuilder* LODManagerToolBuilder = NewObject<ULODManagerToolBuilder>();
+	LODManagerToolBuilder->AssetAPI = ModelingModeAssetGenerationAPI.Get();
+	RegisterTool(ToolManagerCommands.BeginLODManagerTool, TEXT("LODManagerTool"), LODManagerToolBuilder);
+
 	UGenerateStaticMeshLODAssetToolBuilder* GenerateSMLODToolBuilder = NewObject<UGenerateStaticMeshLODAssetToolBuilder>();
 	GenerateSMLODToolBuilder->AssetAPI = ModelingModeAssetGenerationAPI.Get();
 	RegisterTool(ToolManagerCommands.BeginGenerateStaticMeshLODAssetTool, TEXT("GenerateSMLODTool"), GenerateSMLODToolBuilder);
diff --git a/Engine/Plugins/Experimental/ModelingToolsEditorMode/Source/ModelingToolsEditorMode/Private/ModelingToolsEditorModeToolkit.cpp b/Engine/Plugins/Experimental/ModelingToolsEditorMode/Source/ModelingToolsEditorMode/Private/ModelingToolsEditorModeToolkit.cpp
--- a/Engine/Plugins/Experimental/ModelingToolsEditorMode/Source/ModelingToolsEditorMode/Private/ModelingToolsEditorModeToolkit.cpp
+++ b/Engine/Plugins/Experimental/ModelingToolsEditorMode/Source/ModelingToolsEditorMode/Private/ModelingToolsEditorModeToolkit.cpp
@@ -409,9 +409,10 @@ static const FName PrototypesTabName(TEXT("Prototypes"));
 static const FName HairTabName(TEXT("Hair"));
 static const FName PolyEditTabName(TEXT("PolyEdit"));
 static const FName VoxToolsTabName(TEXT("VoxOps"));
+static const FName LODToolsTabName(TEXT("LODs"));
 
 
-const TArray<FName> FModelingToolsEditorModeToolkit::PaletteNames_Standard = { PrimitiveTabName, CreateTabName, PolyModelingTabName, TriModelingTabName, DeformTabName, TransformTabName, MeshProcessingTabName, VoxToolsTabName, AttributesTabName, UVTabName, VolumesTabName, HairTabName };
+const TArray<FName> FModelingToolsEditorModeToolkit::PaletteNames_Standard = { PrimitiveTabName, CreateTabName, PolyModelingTabName, TriModelingTabName, DeformTabName, TransformTabName, MeshProcessingTabName, VoxToolsTabName, AttributesTabName, UVTabName, VolumesTabName, LODToolsTabName, HairTabName };
 
 
 void FModelingToolsEditorModeToolkit::GetToolPaletteNames(TArray<FName>& PaletteNames) const
@@ -504,6 +505,10 @@ void FModelingToolsEditorModeToolkit::BuildToolPalette_Experimental(FName Palett
 		ToolbarBuilder.AddToolBarButton(Commands.BeginRemoveOccludedTrianglesTool);
 		ToolbarBuilder.AddToolBarButton(Commands.BeginSelfUnionTool);
 		ToolbarBuilder.AddToolBarButton(Commands.BeginProjectToTargetTool);
+	}
+	else if (PaletteIndex == LODToolsTabName)
+	{
+		ToolbarBuilder.AddToolBarButton(Commands.BeginLODManagerTool);
 		ToolbarBuilder.AddToolBarButton(Commands.BeginGenerateStaticMeshLODAssetTool);
 	}
 	else if (PaletteIndex == VoxToolsTabName)
diff --git a/Engine/Plugins/Experimental/ModelingToolsEditorMode/Source/ModelingToolsEditorMode/Private/ModelingToolsManagerActions.cpp b/Engine/Plugins/Experimental/ModelingToolsEditorMode/Source/ModelingToolsEditorMode/Private/ModelingToolsManagerActions.cpp
--- a/Engine/Plugins/Experimental/ModelingToolsEditorMode/Source/ModelingToolsEditorMode/Private/ModelingToolsManagerActions.cpp
+++ b/Engine/Plugins/Experimental/ModelingToolsEditorMode/Source/ModelingToolsEditorMode/Private/ModelingToolsManagerActions.cpp
@@ -110,6 +110,8 @@ void FModelingToolsManagerCommands::RegisterCommands()
 	UI_COMMAND(BeginGroomCardsEditorTool, "CardsEd", "Edit Hair Cards", EUserInterfaceActionType::ToggleButton, FInputChord());
 	UI_COMMAND(BeginSubdividePolyTool, "SubD", "Subdivide Mesh Polygroups or Triangles", EUserInterfaceActionType::ToggleButton, FInputChord());
 
+	
+	UI_COMMAND(BeginLODManagerTool, "LODMgr", "Static Mesh Asset LOD Manager", EUserInterfaceActionType::ToggleButton, FInputChord());
 	UI_COMMAND(BeginGenerateStaticMeshLODAssetTool, "AutoLOD", "Generate Static Mesh LOD Asset", EUserInterfaceActionType::ToggleButton, FInputChord());
 	UI_COMMAND(BeginMeshGroupPaintTool, "GrpPnt", "Paint New Mesh Polygroups", EUserInterfaceActionType::ToggleButton, FInputChord());
 
diff --git a/Engine/Plugins/Experimental/ModelingToolsEditorMode/Source/ModelingToolsEditorMode/Public/ModelingToolsManagerActions.h b/Engine/Plugins/Experimental/ModelingToolsEditorMode/Source/ModelingToolsEditorMode/Public/ModelingToolsManagerActions.h
--- a/Engine/Plugins/Experimental/ModelingToolsEditorMode/Source/ModelingToolsEditorMode/Public/ModelingToolsManagerActions.h
+++ b/Engine/Plugins/Experimental/ModelingToolsEditorMode/Source/ModelingToolsEditorMode/Public/ModelingToolsManagerActions.h
@@ -104,6 +104,7 @@ public:
 	TSharedPtr<FUICommandInfo> BeginGroomMeshToCardsTool;
 	TSharedPtr<FUICommandInfo> BeginGroomCardsEditorTool;
 
+	TSharedPtr<FUICommandInfo> BeginLODManagerTool;
 	TSharedPtr<FUICommandInfo> BeginGenerateStaticMeshLODAssetTool;
 
 	TSharedPtr<FUICommandInfo> BeginPolyModelTool_FaceSelect;
