LODTools: disconnect lines actor on LODManagerTool shutdown
#rnx

diff --git a/Engine/Plugins/Experimental/MeshLODToolset/Source/MeshLODToolset/Private/Tools/LODManagerTool.cpp b/Engine/Plugins/Experimental/MeshLODToolset/Source/MeshLODToolset/Private/Tools/LODManagerTool.cpp
--- a/Engine/Plugins/Experimental/MeshLODToolset/Source/MeshLODToolset/Private/Tools/LODManagerTool.cpp
+++ b/Engine/Plugins/Experimental/MeshLODToolset/Source/MeshLODToolset/Private/Tools/LODManagerTool.cpp
@@ -133,6 +133,7 @@ void ULODManagerTool::Shutdown(EToolShutdownType ShutdownType)
 	if (LODPreview)
 	{
 		LODPreview->Disconnect();
+		LODPreviewLines->Disconnect();
 	}
 	for (TUniquePtr<FPrimitiveComponentTarget>& Target : ComponentTargets)
 	{
