LODTools: add missing early-out when input error is detected
#rnx

diff --git a/Engine/Plugins/Experimental/MeshLODToolset/Source/MeshLODToolset/Private/Tools/LODManagerTool.cpp b/Engine/Plugins/Experimental/MeshLODToolset/Source/MeshLODToolset/Private/Tools/LODManagerTool.cpp
--- a/Engine/Plugins/Experimental/MeshLODToolset/Source/MeshLODToolset/Private/Tools/LODManagerTool.cpp
+++ b/Engine/Plugins/Experimental/MeshLODToolset/Source/MeshLODToolset/Private/Tools/LODManagerTool.cpp
@@ -377,6 +377,7 @@ void ULODManagerTool::UpdatePreviewLOD()
 		// badness
 		LODPreview->SetVisible(false);
 		ComponentTargets[0]->SetOwnerVisibility(true);
+		return;
 	}
 
 	TUniquePtr<FLODMeshInfo>* Found = LODMeshCache.Find(SelectedLOD);
