Fix a crash when running AutoLOD twice. Don't delete existng assets if we're in "Create new asset" mode.

#rnx

diff --git a/Engine/Plugins/Experimental/MeshLODToolset/Source/MeshLODToolset/Private/Graphs/GenerateStaticMeshLODProcess.cpp b/Engine/Plugins/Experimental/MeshLODToolset/Source/MeshLODToolset/Private/Graphs/GenerateStaticMeshLODProcess.cpp
--- a/Engine/Plugins/Experimental/MeshLODToolset/Source/MeshLODToolset/Private/Graphs/GenerateStaticMeshLODProcess.cpp
+++ b/Engine/Plugins/Experimental/MeshLODToolset/Source/MeshLODToolset/Private/Graphs/GenerateStaticMeshLODProcess.cpp
@@ -57,6 +57,40 @@ struct FReadTextureJob
 };
 
 
+namespace GenerateStaticMeshLODProcessHelpers
+{
+	// Given "xxxxx" returns "xxxxx_1"
+	// Given "xxxxx_1" returns "xxxxx_2"
+	// etc.
+	// If anything goes wrong, just add "_1" to the end
+	void AppendOrIncrementSuffix(FString& S)
+	{
+		TArray<FString> Substrings;
+		S.ParseIntoArray(Substrings, TEXT("_"));
+		if ( Substrings.Num() <= 1 ) 
+		{
+			S += TEXT("_1");
+		}
+		else
+		{
+			FString LastSubstring = Substrings.Last();
+			int32 Num;
+			bool bParsed = LexTryParseString(Num, *LastSubstring);
+			if (bParsed)
+			{
+				++Num;
+				Substrings.RemoveAt(Substrings.Num() - 1);
+				S = FString::Join(Substrings, TEXT("_"));
+				S += TEXT("_") + FString::FromInt(Num);
+			}
+			else
+			{
+				S += TEXT("_1");
+			}
+		}
+	}
+}
+
 bool UGenerateStaticMeshLODProcess::Initialize(UStaticMesh* StaticMeshIn)
 {
 	if (!ensure(StaticMeshIn)) return false;
@@ -619,9 +653,11 @@ bool UGenerateStaticMeshLODProcess::WriteDerivedAssetData()
 {
 	AllDerivedTextures.Reset();
 
-	WriteDerivedTextures();
+	constexpr bool bCreatingNewStaticMeshAsset = true;
 
-	WriteDerivedMaterials();
+	WriteDerivedTextures(bCreatingNewStaticMeshAsset);
+
+	WriteDerivedMaterials(bCreatingNewStaticMeshAsset);
 
 	WriteDerivedStaticMeshAsset();
 
@@ -636,9 +672,11 @@ void UGenerateStaticMeshLODProcess::UpdateSourceAsset(bool bSetNewHDSourceAsset)
 {
 	AllDerivedTextures.Reset();
 
-	WriteDerivedTextures();
+	constexpr bool bCreatingNewStaticMeshAsset = false;
+
+	WriteDerivedTextures(bCreatingNewStaticMeshAsset);
 
-	WriteDerivedMaterials();
+	WriteDerivedMaterials(bCreatingNewStaticMeshAsset);
 
 	UpdateSourceStaticMeshAsset(bSetNewHDSourceAsset);
 
@@ -647,7 +685,7 @@ void UGenerateStaticMeshLODProcess::UpdateSourceAsset(bool bSetNewHDSourceAsset)
 }
 
 
-void UGenerateStaticMeshLODProcess::WriteDerivedTextures()
+void UGenerateStaticMeshLODProcess::WriteDerivedTextures(bool bCreatingNewStaticMeshAsset)
 {
 	IAssetTools& AssetTools = FModuleManager::Get().LoadModuleChecked<FAssetToolsModule>("AssetTools").Get();
 
@@ -701,7 +739,7 @@ void UGenerateStaticMeshLODProcess::WriteDerivedTextures()
 				FTexture2DBuilder::CopyPlatformDataToSourceData(DerivedTex.Texture, FTexture2DBuilder::ETextureType::Color);
 
 				// write asset
-				bool bWriteOK = WriteDerivedTexture(SourceTex.Texture, DerivedTex.Texture);
+				bool bWriteOK = WriteDerivedTexture(SourceTex.Texture, DerivedTex.Texture, bCreatingNewStaticMeshAsset);
 				ensure(bWriteOK);
 
 				SourceTexturesWritten.Add(SourceTex.Texture);
@@ -724,7 +762,7 @@ void UGenerateStaticMeshLODProcess::WriteDerivedTextures()
 			FTexture2DBuilder::CopyPlatformDataToSourceData(DerivedNormalMapTex, FTexture2DBuilder::ETextureType::NormalMap);
 
 			// write asset
-			bool bWriteOK = WriteDerivedTexture(DerivedNormalMapTex, SourceAssetName + TEXT("_NormalMap"));
+			bool bWriteOK = WriteDerivedTexture(DerivedNormalMapTex, SourceAssetName + TEXT("_NormalMap"), bCreatingNewStaticMeshAsset);
 			ensure(bWriteOK);
 		}
 	}
@@ -734,29 +772,44 @@ void UGenerateStaticMeshLODProcess::WriteDerivedTextures()
 
 
 
-bool UGenerateStaticMeshLODProcess::WriteDerivedTexture(UTexture2D* SourceTexture, UTexture2D* DerivedTexture)
+bool UGenerateStaticMeshLODProcess::WriteDerivedTexture(UTexture2D* SourceTexture, UTexture2D* DerivedTexture, bool bCreatingNewStaticMeshAsset)
 {
 	IAssetTools& AssetTools = FModuleManager::Get().LoadModuleChecked<FAssetToolsModule>("AssetTools").Get();
 
 	FString SourceTexPath = UEditorAssetLibrary::GetPathNameForLoadedAsset(SourceTexture);
 	FString TexName = FPaths::GetBaseFilename(SourceTexPath, true);
-	return WriteDerivedTexture(DerivedTexture, TexName);
+	return WriteDerivedTexture(DerivedTexture, TexName, bCreatingNewStaticMeshAsset);
 }
 
 
-bool UGenerateStaticMeshLODProcess::WriteDerivedTexture(UTexture2D* DerivedTexture, FString BaseTexName)
+
+bool UGenerateStaticMeshLODProcess::WriteDerivedTexture(UTexture2D* DerivedTexture, FString BaseTexName, bool bCreatingNewStaticMeshAsset)
 {
 	IAssetTools& AssetTools = FModuleManager::Get().LoadModuleChecked<FAssetToolsModule>("AssetTools").Get();
 
 	FString NewTexName = FString::Printf(TEXT("%s%s"), *BaseTexName, *DerivedSuffix);
 	FString NewAssetPath = FPaths::Combine(DerivedAssetFolder, NewTexName);
 
-	// delete existing asset so that we can have a clean duplicate
-	bool bNewAssetExists = UEditorAssetLibrary::DoesAssetExist(NewAssetPath);
-	if (bNewAssetExists)
+	if (bCreatingNewStaticMeshAsset)
+	{
+		// Don't delete an existing asset. If name collision occurs, rename the new asset.
+		bool bNewAssetExists = UEditorAssetLibrary::DoesAssetExist(NewAssetPath);
+		while (bNewAssetExists)
+		{
+			GenerateStaticMeshLODProcessHelpers::AppendOrIncrementSuffix(NewTexName);
+			NewAssetPath = FPaths::Combine(DerivedAssetFolder, NewTexName);
+			bNewAssetExists = UEditorAssetLibrary::DoesAssetExist(NewAssetPath);
+		}
+	}
+	else
 	{
-		bool bDeleteOK = UEditorAssetLibrary::DeleteAsset(NewAssetPath);
-		ensure(bDeleteOK);
+		// Modifying the static mesh in place. Delete existing asset so that we can have a clean duplicate
+		bool bNewAssetExists = UEditorAssetLibrary::DoesAssetExist(NewAssetPath);
+		if (bNewAssetExists)
+		{
+			bool bDeleteOK = UEditorAssetLibrary::DeleteAsset(NewAssetPath);
+			ensure(bDeleteOK);
+		}
 	}
 
 	// create package
@@ -782,7 +835,7 @@ bool UGenerateStaticMeshLODProcess::WriteDerivedTexture(UTexture2D* DerivedTextu
 }
 
 
-void UGenerateStaticMeshLODProcess::WriteDerivedMaterials()
+void UGenerateStaticMeshLODProcess::WriteDerivedMaterials(bool bCreatingNewStaticMeshAsset)
 {
 	IAssetTools& AssetTools = FModuleManager::Get().LoadModuleChecked<FAssetToolsModule>("AssetTools").Get();
 
@@ -811,12 +864,26 @@ void UGenerateStaticMeshLODProcess::WriteDerivedMaterials()
 		FString NewMaterialName = FString::Printf(TEXT("%s%s"), *MaterialName, *DerivedSuffix);
 		FString NewMaterialPath = FPaths::Combine(DerivedAssetFolder, NewMaterialName);
 
-		// delete existing material so that we can have a clean duplicate
-		bool bNewMaterialExists = UEditorAssetLibrary::DoesAssetExist(NewMaterialPath);
-		if (bNewMaterialExists)
+		if (bCreatingNewStaticMeshAsset)
 		{
-			bool bDeleteOK = UEditorAssetLibrary::DeleteAsset(NewMaterialPath);
-			ensure(bDeleteOK);
+			// Don't delete an existing material. If name collision occurs, rename the new material.
+			bool bNewAssetExists = UEditorAssetLibrary::DoesAssetExist(NewMaterialPath);
+			while (bNewAssetExists)
+			{
+				GenerateStaticMeshLODProcessHelpers::AppendOrIncrementSuffix(NewMaterialName);
+				NewMaterialPath = FPaths::Combine(DerivedAssetFolder, NewMaterialName);
+				bNewAssetExists = UEditorAssetLibrary::DoesAssetExist(NewMaterialPath);
+			}
+		}
+		else
+		{
+			// Modifying the static mesh in place. Delete existing asset so that we can have a clean duplicate
+			bool bNewAssetExists = UEditorAssetLibrary::DoesAssetExist(NewMaterialPath);
+			if (bNewAssetExists)
+			{
+				bool bDeleteOK = UEditorAssetLibrary::DeleteAsset(NewMaterialPath);
+				ensure(bDeleteOK);
+			}
 		}
 
 		// If source is a MIC, we can just duplicate it. If it is a UMaterial, we want to
@@ -922,17 +989,16 @@ void UGenerateStaticMeshLODProcess::WriteDerivedStaticMeshAsset()
 	int32 NumMaterials = SourceMaterials.Num();
 	for (int32 mi = 0; mi < NumMaterials; ++mi)
 	{
-		if (SourceMaterials[mi].bIsPreviouslyGeneratedMaterial)
-		{
-			NewMaterials.Add(SourceMaterials[mi].SourceMaterial);		// TODO: should be skipping this material...
-		}
-		else if (SourceMaterials[mi].bIsReusable)
-		{
-			NewMaterials.Add(SourceMaterials[mi].SourceMaterial);
-		}
-		else 
+		if (!SourceMaterials[mi].bIsPreviouslyGeneratedMaterial)	// Skip previously generated
 		{
-			NewMaterials.Add(DerivedMaterials[mi].DerivedMaterial);
+			if (SourceMaterials[mi].bIsReusable)
+			{
+				NewMaterials.Add(SourceMaterials[mi].SourceMaterial);
+			}
+			else
+			{
+				NewMaterials.Add(DerivedMaterials[mi].DerivedMaterial);
+			}
 		}
 	}
 
diff --git a/Engine/Plugins/Experimental/MeshLODToolset/Source/MeshLODToolset/Public/Graphs/GenerateStaticMeshLODProcess.h b/Engine/Plugins/Experimental/MeshLODToolset/Source/MeshLODToolset/Public/Graphs/GenerateStaticMeshLODProcess.h
--- a/Engine/Plugins/Experimental/MeshLODToolset/Source/MeshLODToolset/Public/Graphs/GenerateStaticMeshLODProcess.h
+++ b/Engine/Plugins/Experimental/MeshLODToolset/Source/MeshLODToolset/Public/Graphs/GenerateStaticMeshLODProcess.h
@@ -292,10 +292,10 @@ protected:
 
 	FGenerateStaticMeshLODProcessSettings CurrentSettings;
 
-	bool WriteDerivedTexture(UTexture2D* SourceTexture, UTexture2D* DerivedTexture);
-	bool WriteDerivedTexture(UTexture2D* DerivedTexture, FString BaseTexName);
-	void WriteDerivedTextures();
-	void WriteDerivedMaterials();
+	bool WriteDerivedTexture(UTexture2D* SourceTexture, UTexture2D* DerivedTexture, bool bCreatingNewStaticMeshAsset);
+	bool WriteDerivedTexture(UTexture2D* DerivedTexture, FString BaseTexName, bool bCreatingNewStaticMeshAsset);
+	void WriteDerivedTextures(bool bCreatingNewStaticMeshAsset);
+	void WriteDerivedMaterials(bool bCreatingNewStaticMeshAsset);
 	void UpdateMaterialTextureParameters(UMaterialInstanceConstant* Material, FDerivedMaterialInfo& DerivedMaterialInfo);
 	void WriteDerivedStaticMeshAsset();
 
