Set mins and maxes for AutoLOD tool properties.

#rnx

diff --git a/Engine/Plugins/Experimental/MeshLODToolset/Source/MeshLODToolset/Public/Graphs/GenerateStaticMeshLODProcess.h b/Engine/Plugins/Experimental/MeshLODToolset/Source/MeshLODToolset/Public/Graphs/GenerateStaticMeshLODProcess.h
--- a/Engine/Plugins/Experimental/MeshLODToolset/Source/MeshLODToolset/Public/Graphs/GenerateStaticMeshLODProcess.h
+++ b/Engine/Plugins/Experimental/MeshLODToolset/Source/MeshLODToolset/Public/Graphs/GenerateStaticMeshLODProcess.h
@@ -75,15 +75,15 @@ struct FGenerateStaticMeshLODProcessSettings
 	UPROPERTY(EditAnywhere, Category = DetailFilter, meta = (DisplayName = "Thicken weight map name"))
 	FName ThickenWeightMapName = FName(TEXT("ThickenWeightMap"));
 
-	UPROPERTY(EditAnywhere, Category = DetailFilter)
+	UPROPERTY(EditAnywhere, Category = DetailFilter, meta = (UIMin = "0", UIMax = "100", ClampMin = "0", ClampMax = "1000"))
 	float ThickenAmount = 0.0f;
 
 	// Solidify settings
 
-	UPROPERTY(EditAnywhere, Category = Solidify, meta = (DisplayName="Voxel Resolution"))
+	UPROPERTY(EditAnywhere, Category = Solidify, meta = (UIMin = "8", UIMax = "1024", ClampMin = "8", ClampMax = "1024", DisplayName="Voxel Resolution"))
 	int SolidifyVoxelResolution = 64;
 
-	UPROPERTY(EditAnywhere, Category = Solidify, AdvancedDisplay)
+	UPROPERTY(EditAnywhere, Category = Solidify, AdvancedDisplay, meta = (UIMin = "0.1", UIMax = ".9", ClampMin = "-10", ClampMax = "10"))
 	float WindingThreshold = 0.5f;
 
 
@@ -92,13 +92,13 @@ struct FGenerateStaticMeshLODProcessSettings
 	//UPROPERTY(EditAnywhere, Category = Morphology, meta = (DisplayName = "Voxel Resolution"))
 	//int MorphologyVoxelResolution = 64;
 
-	UPROPERTY(EditAnywhere, Category = Morphology, meta = (DisplayName = "Closure Distance"))
+	UPROPERTY(EditAnywhere, Category = Morphology, meta = (UIMin = "0", UIMax = "100", ClampMin = "0", ClampMax = "1000", DisplayName = "Closure Distance"))
 	float ClosureDistance = 1.0f;
 
 
 	// Simplify settings
 
-	UPROPERTY(EditAnywhere, Category = Simplify, meta = (DisplayName = "Simplify Tri Count"))
+	UPROPERTY(EditAnywhere, Category = Simplify, meta = (UIMin = "1", ClampMin = "1", DisplayName = "Simplify Tri Count"))
 	int SimplifyTriangleCount = 500;
 
 
@@ -113,7 +113,7 @@ struct FGenerateStaticMeshLODProcessSettings
 	UPROPERTY(EditAnywhere, Category = Baking , meta = (DisplayName = "Bake Image Res"))
 	EGenerateStaticMeshLODBakeResolution BakeResolution = EGenerateStaticMeshLODBakeResolution::Resolution512;
 
-	UPROPERTY(EditAnywhere, Category = Baking, meta = (DisplayName = "Bake Thickness"))
+	UPROPERTY(EditAnywhere, Category = Baking, meta = (UIMin = "0", UIMax = "100", ClampMin = "0", ClampMax = "1000", DisplayName = "Bake Thickness"))
 	float BakeThickness = 5.0f;
 
 	// Transient property, not set directly by the user. The user controls a CollisionGroupLayerName dropdown property
