Fix a crash when applying AutoLOD twice

#rnx

diff --git a/Engine/Plugins/Experimental/MeshLODToolset/Source/MeshLODToolset/Private/Graphs/GenerateStaticMeshLODProcess.cpp b/Engine/Plugins/Experimental/MeshLODToolset/Source/MeshLODToolset/Private/Graphs/GenerateStaticMeshLODProcess.cpp
--- a/Engine/Plugins/Experimental/MeshLODToolset/Source/MeshLODToolset/Private/Graphs/GenerateStaticMeshLODProcess.cpp
+++ b/Engine/Plugins/Experimental/MeshLODToolset/Source/MeshLODToolset/Private/Graphs/GenerateStaticMeshLODProcess.cpp
@@ -1174,7 +1174,12 @@ void UGenerateStaticMeshLODProcess::UpdateSourceStaticMeshAsset(bool bSetNewHDSo
 	TPolygonGroupAttributesRef<FName> PolygonGroupImportedMaterialSlotNames = Attributes.GetPolygonGroupMaterialSlotNames();
 	for (int32 SlotIdx : DerivedMaterialSlotIndices)
 	{
-		PolygonGroupImportedMaterialSlotNames.Set(SlotIdx, NewMaterialSet[SlotIdx].ImportedMaterialSlotName);
+		// It's possible that NewMaterialSet.Num() > PolygonGroupImportedMaterialSlotNames.GetNumElements()
+		// if there are new materials that aren't referenced by any triangles...
+		if (SlotIdx < PolygonGroupImportedMaterialSlotNames.GetNumElements())		
+		{
+			PolygonGroupImportedMaterialSlotNames.Set(SlotIdx, NewMaterialSet[SlotIdx].ImportedMaterialSlotName);
+		}
 	}
 
 	// Disable auto-generated normals/tangents, we need to use the ones we computed in LOD Generator
