AutoLOD: detect when an input material has a non-parameter texture sampler and issue a warning.

#rnx

diff --git a/Engine/Plugins/Experimental/MeshLODToolset/Source/MeshLODToolset/Private/Graphs/GenerateStaticMeshLODProcess.cpp b/Engine/Plugins/Experimental/MeshLODToolset/Source/MeshLODToolset/Private/Graphs/GenerateStaticMeshLODProcess.cpp
--- a/Engine/Plugins/Experimental/MeshLODToolset/Source/MeshLODToolset/Private/Graphs/GenerateStaticMeshLODProcess.cpp
+++ b/Engine/Plugins/Experimental/MeshLODToolset/Source/MeshLODToolset/Private/Graphs/GenerateStaticMeshLODProcess.cpp
@@ -29,6 +29,8 @@
 #include "Materials/MaterialInstanceConstant.h"
 #include "Factories/MaterialInstanceConstantFactoryNew.h"
 #include "Materials/MaterialInstanceDynamic.h"
+#include "MaterialGraph/MaterialGraph.h"
+#include "MaterialGraph/MaterialGraphNode.h"
 #include "WeightMapUtil.h"
 
 #include "Engine/StaticMesh.h"
@@ -131,6 +133,52 @@ bool UGenerateStaticMeshLODProcess::Initialize(UStaticMesh* StaticMeshIn)
 	{
 		SourceMaterials[mi].SourceMaterial = Materials[mi];
 
+
+		// detect hard-coded (non-parameter) texture samples
+		{
+			UMaterial* Material = Materials[mi].MaterialInterface->GetMaterial();
+
+			// go over the nodes in the material graph looking for texture samples
+			const UMaterialGraph* MatGraph = Material->MaterialGraph;
+
+			if (MatGraph == nullptr)	// create a material graph from the material if necessary
+			{
+				UMaterialGraph* NewMatGraph = CastChecked<UMaterialGraph>(NewObject<UEdGraph>(Material, UMaterialGraph::StaticClass(), NAME_None, RF_Transient));
+				NewMatGraph->Material = Material;
+				NewMatGraph->RebuildGraph();
+				MatGraph = NewMatGraph;
+			}
+
+			bool bFoundTextureNonParamExpession = false;
+			const TArray<TObjectPtr<class UEdGraphNode>>& Nodes = MatGraph->Nodes;
+			for (const TObjectPtr<UEdGraphNode>& Node : Nodes)
+			{
+				const UMaterialGraphNode* GraphNode = Cast<UMaterialGraphNode>(Node);
+				if (GraphNode)
+				{
+					const UMaterialExpressionTextureSampleParameter* TextureSampleParameterBase = Cast<UMaterialExpressionTextureSampleParameter>(GraphNode->MaterialExpression);
+					if (!TextureSampleParameterBase)
+					{
+						const UMaterialExpressionTextureSample* TextureSampleBase = Cast<UMaterialExpressionTextureSample>(GraphNode->MaterialExpression);
+						if (TextureSampleBase)
+						{
+							// node is UMaterialExpressionTextureSample but not UMaterialExpressionTextureSampleParameter
+							bFoundTextureNonParamExpession = true;
+							break;
+						}
+					}
+
+				}
+			}
+			if (bFoundTextureNonParamExpession)
+			{
+				// TODO: bubble this up to the tool level
+				UE_LOG(LogTemp, Warning, 
+					   TEXT("UGenerateStaticMeshLODProcess: Non-parameter texture sampler detected in input material [%s]. Output materials may have unexpected behaviour."),
+					   *Material->GetName());
+			}
+		}
+
 		UMaterialInterface* MaterialInterface = Materials[mi].MaterialInterface;
 		if (MaterialInterface == nullptr)
 		{
