AutoLOD: turn on convex hull subsampling by default. Add a "None" option to collision geometry.

#rnx

diff --git a/Engine/Plugins/Experimental/GeometryFlow/Source/GeometryFlowMeshProcessing/Private/PhysicsNodes/GenerateSimpleCollisionNode.cpp b/Engine/Plugins/Experimental/GeometryFlow/Source/GeometryFlowMeshProcessing/Private/PhysicsNodes/GenerateSimpleCollisionNode.cpp
--- a/Engine/Plugins/Experimental/GeometryFlow/Source/GeometryFlowMeshProcessing/Private/PhysicsNodes/GenerateSimpleCollisionNode.cpp
+++ b/Engine/Plugins/Experimental/GeometryFlow/Source/GeometryFlowMeshProcessing/Private/PhysicsNodes/GenerateSimpleCollisionNode.cpp
@@ -175,6 +175,8 @@ EGeometryFlowResult FGenerateSimpleCollisionNode::EvaluateInternal(const FDynami
 	 case ESimpleCollisionGeometryType::MinVolume:
 		 ShapeApproximator.Generate_MinVolume(OutCollisionGeometry.Geometry);
 		 break;
+	 case ESimpleCollisionGeometryType::None:
+		 break;
 	 }
 
 	 if (EvaluationInfo && EvaluationInfo->Progress && EvaluationInfo->Progress->Cancelled())
diff --git a/Engine/Plugins/Experimental/GeometryFlow/Source/GeometryFlowMeshProcessing/Public/PhysicsNodes/GenerateConvexHullsCollisionNode.h b/Engine/Plugins/Experimental/GeometryFlow/Source/GeometryFlowMeshProcessing/Public/PhysicsNodes/GenerateConvexHullsCollisionNode.h
--- a/Engine/Plugins/Experimental/GeometryFlow/Source/GeometryFlowMeshProcessing/Public/PhysicsNodes/GenerateConvexHullsCollisionNode.h
+++ b/Engine/Plugins/Experimental/GeometryFlow/Source/GeometryFlowMeshProcessing/Public/PhysicsNodes/GenerateConvexHullsCollisionNode.h
@@ -23,7 +23,7 @@ struct GEOMETRYFLOWMESHPROCESSING_API FGenerateConvexHullsCollisionSettings
 	DECLARE_GEOMETRYFLOW_DATA_TYPE_IDENTIFIER(EMeshProcessingDataTypes::GenerateCollisionConvexHullsSettings);
 
 	int32 SimplifyToTriangleCount = 50;
-	bool bPrefilterVertices = false;
+	bool bPrefilterVertices = true;
 	int PrefilterGridResolution = 10;
 	
 };
diff --git a/Engine/Plugins/Experimental/GeometryFlow/Source/GeometryFlowMeshProcessing/Public/PhysicsNodes/GenerateSimpleCollisionNode.h b/Engine/Plugins/Experimental/GeometryFlow/Source/GeometryFlowMeshProcessing/Public/PhysicsNodes/GenerateSimpleCollisionNode.h
--- a/Engine/Plugins/Experimental/GeometryFlow/Source/GeometryFlowMeshProcessing/Public/PhysicsNodes/GenerateSimpleCollisionNode.h
+++ b/Engine/Plugins/Experimental/GeometryFlow/Source/GeometryFlowMeshProcessing/Public/PhysicsNodes/GenerateSimpleCollisionNode.h
@@ -28,7 +28,8 @@ enum class ESimpleCollisionGeometryType : uint8
 	Capsules,
 	ConvexHulls,
 	SweptHulls,
-	MinVolume
+	MinVolume,
+	None
 };
 
 
@@ -41,7 +42,7 @@ struct GEOMETRYFLOWMESHPROCESSING_API FGenerateSimpleCollisionSettings
 	struct FGenerateConvexHullSettings
 	{
 		int32 SimplifyToTriangleCount = 50;
-		bool bPrefilterVertices = false;
+		bool bPrefilterVertices = true;
 		int PrefilterGridResolution = 10;
 	} ConvexHullSettings;
 
diff --git a/Engine/Plugins/Experimental/MeshLODToolset/Source/MeshLODToolset/Public/Graphs/GenerateStaticMeshLODProcess.h b/Engine/Plugins/Experimental/MeshLODToolset/Source/MeshLODToolset/Public/Graphs/GenerateStaticMeshLODProcess.h
--- a/Engine/Plugins/Experimental/MeshLODToolset/Source/MeshLODToolset/Public/Graphs/GenerateStaticMeshLODProcess.h
+++ b/Engine/Plugins/Experimental/MeshLODToolset/Source/MeshLODToolset/Public/Graphs/GenerateStaticMeshLODProcess.h
@@ -33,18 +33,19 @@ enum class EGenerateStaticMeshLODBakeResolution
 	Resolution8192 = 8192 UMETA(DisplayName = "8192 x 8192")
 };
 
-// NOTE: This must be kept in sync with ESimpleCollisionGeometryType in GenerateSimpleCollisionNode.h
-
 UENUM()
 enum class EGenerateStaticMeshLODSimpleCollisionGeometryType : uint8
 {
+	// NOTE: This must be kept in sync with ESimpleCollisionGeometryType in GenerateSimpleCollisionNode.h
+
 	AlignedBoxes,
 	OrientedBoxes,
 	MinimalSpheres,
 	Capsules,
 	ConvexHulls,
 	SweptHulls,
-	MinVolume
+	MinVolume,
+	None
 };
 
 // NOTE: This must be kept in sync with FMeshSimpleShapeApproximation::EProjectedHullAxisMode in MeshSimpleShapeApproximation.h
