Remove previous parallel executor for AutoLOD. FGraph now uses a Async/Future system by default, which is typically faster and more flexible.

#rnx

diff --git a/Engine/Plugins/Experimental/GeometryFlow/Source/GeometryFlowCore/Private/GeometryFlowExecutor.cpp b/Engine/Plugins/Experimental/GeometryFlow/Source/GeometryFlowCore/Private/GeometryFlowExecutor.cpp
deleted file mode 100644
diff --git a/Engine/Plugins/Experimental/GeometryFlow/Source/GeometryFlowCore/Public/GeometryFlowExecutor.h b/Engine/Plugins/Experimental/GeometryFlow/Source/GeometryFlowCore/Public/GeometryFlowExecutor.h
deleted file mode 100644
diff --git a/Engine/Plugins/Experimental/MeshLODToolset/Source/MeshLODToolset/Private/Graphs/GenerateMeshLODGraph.cpp b/Engine/Plugins/Experimental/MeshLODToolset/Source/MeshLODToolset/Private/Graphs/GenerateMeshLODGraph.cpp
--- a/Engine/Plugins/Experimental/MeshLODToolset/Source/MeshLODToolset/Private/Graphs/GenerateMeshLODGraph.cpp
+++ b/Engine/Plugins/Experimental/MeshLODToolset/Source/MeshLODToolset/Private/Graphs/GenerateMeshLODGraph.cpp
@@ -28,7 +28,6 @@
 
 #include "MeshDecompositionNodes/MakeTriangleSetsNode.h"
 #include "PhysicsNodes/GenerateSimpleCollisionNode.h"
-#include "GeometryFlowExecutor.h"
 
 #include "ExplicitUseGeometryMathTypes.h"		// using UE::Geometry::(math types)
 using namespace UE::Geometry;
@@ -41,190 +40,6 @@ void FGenerateMeshLODGraph::SetSourceMesh(const FDynamicMesh3& SourceMeshIn)
 }
 
 
-#if 0
-
-void FGenerateMeshLODGraph::EvaluateResultParallel(
-	FDynamicMesh3& ResultMesh,
-	FMeshTangentsd& ResultTangents,
-	FSimpleShapeSet3d& ResultCollision,
-	UE::GeometryFlow::FNormalMapImage& NormalMap,
-	TArray<TUniquePtr<UE::GeometryFlow::FTextureImage>>& TextureImages)
-{
-	//FScopedDurationTimeLogger Timer(TEXT("FGenerateMeshLODGraph::EvaluateResult -- parallel execution"));
-
-	TArray<FGeometryFlowExecutor::NodeOutputSpec> DesiredOutputs;
-
-	bool bTakeNormalMap = true;
-	DesiredOutputs.Add(FGeometryFlowExecutor::NodeOutputSpec{ 
-		BakeNormalMapNode, 
-		FBakeMeshNormalMapNode::OutParamNormalMap(),
-		bTakeNormalMap });
-
-	bool bTakeTexBake = true;
-	for (FBakeTextureGraphInfo& TexBakeStep : BakeTextureNodes)
-	{
-		DesiredOutputs.Add(FGeometryFlowExecutor::NodeOutputSpec{
-			TexBakeStep.BakeNode,
-			FBakeMeshTextureImageNode::OutParamTextureImage(),
-			bTakeTexBake });
-	}
-
-	bool bTakeResultCollision = false;
-	DesiredOutputs.Add(FGeometryFlowExecutor::NodeOutputSpec{
-			CollisionOutputNode,
-			FCollisionGeometryTransferNode::OutParamValue(),
-			bTakeResultCollision });
-
-	bool bTakeResultTangents = false;
-	DesiredOutputs.Add(FGeometryFlowExecutor::NodeOutputSpec{
-		TangentsOutputNode,
-		FMeshTangentsTransferNode::OutParamValue(),
-		bTakeResultTangents });
-
-	bool bTakeResultMesh = true;
-	DesiredOutputs.Add(FGeometryFlowExecutor::NodeOutputSpec{
-		MeshOutputNode,
-		FDynamicMeshTransferNode::OutParamValue(),
-		bTakeResultMesh });
-
-	FGeometryFlowExecutor Exec(*Graph);
-	TArray<TSafeSharedPtr<IData>> OutputDatas;
-	Exec.ComputeOutputs(DesiredOutputs, OutputDatas);
-
-	check(OutputDatas.Num() == DesiredOutputs.Num());
-	TArray<TSafeSharedPtr<IData>>::TIterator OutputDataIter = OutputDatas.CreateIterator();
-
-	NormalMap = FNormalMapImage();
-	EGeometryFlowResult ExtractResult = ExtractData(*OutputDataIter,
-													NormalMap,
-													(int)EMeshProcessingDataTypes::NormalMapImage,
-													bTakeNormalMap);
-	check(ExtractResult == EGeometryFlowResult::Ok);
-	++OutputDataIter;
-
-	for (FBakeTextureGraphInfo& TexBakeStep : BakeTextureNodes)
-	{
-		// { TexBakeStep.BakeNode, FBakeMeshTextureImageNode::OutParamTextureImage() }
-		TUniquePtr<UE::GeometryFlow::FTextureImage> NewImage = MakeUnique<UE::GeometryFlow::FTextureImage>();
-		
-		ExtractResult = ExtractData(*OutputDataIter, *NewImage, (int)EMeshProcessingDataTypes::TextureImage, bTakeTexBake);
-		check(ExtractResult == EGeometryFlowResult::Ok);
-		TextureImages.Add(MoveTemp(NewImage));
-		++OutputDataIter;
-	}
-
-	// {CollisionOutputNode, FCollisionGeometryTransferNode::OutParamValue() }
-	ResultCollision = FSimpleShapeSet3d();
-	ExtractResult = ExtractData(*OutputDataIter, ResultCollision, (int)FCollisionGeometry::DataTypeIdentifier, bTakeResultCollision);
-	check(ExtractResult == EGeometryFlowResult::Ok);
-	++OutputDataIter;
-
-	// { TangentsOutputNode, FMeshTangentsTransferNode::OutParamValue() }
-	ResultTangents = FMeshTangentsd();
-	ExtractResult = ExtractData(*OutputDataIter, ResultTangents, (int)EMeshProcessingDataTypes::MeshTangentSet, bTakeResultTangents);
-	check(ExtractResult == EGeometryFlowResult::Ok);
-	++OutputDataIter;
-
-	// {MeshOutputNode, FDynamicMeshTransferNode::OutParamValue() }
-	ResultMesh.Clear();
-	ExtractResult = ExtractData(*OutputDataIter, ResultMesh, (int32)EMeshProcessingDataTypes::DynamicMesh, bTakeResultMesh);
-	check(ExtractResult == EGeometryFlowResult::Ok);
-	++OutputDataIter;
-}
-
-#else
-
-void FGenerateMeshLODGraph::EvaluateResultParallel(
-	FDynamicMesh3& ResultMesh,
-	FMeshTangentsd& ResultTangents,
-	FSimpleShapeSet3d& ResultCollision,
-	UE::GeometryFlow::FNormalMapImage& NormalMap,
-	TArray<TUniquePtr<UE::GeometryFlow::FTextureImage>>& TextureImages,
-	FProgressCancel* Progress)
-{
-	//FScopedDurationTimeLogger Timer(TEXT("FGenerateMeshLODGraph::EvaluateResult -- parallel execution"));
-
-	FGeometryFlowExecutor Exec(*Graph);
-	Exec.AsyncRunGraph(Progress);
-
-	if (Progress && Progress->Cancelled())
-	{
-		return;
-	}
-
-	NormalMap = FNormalMapImage();
-	bool bTakeNormalMap = true;
-	EGeometryFlowResult ExtractResult = Exec.GetOutput(BakeNormalMapNode,
-													   FBakeMeshNormalMapNode::OutParamNormalMap(),
-													   NormalMap, 
-													   (int)EMeshProcessingDataTypes::NormalMapImage, 
-													   bTakeNormalMap);
-	if (ExtractResult == EGeometryFlowResult::OperationCancelled)
-	{
-		return;
-	}
-	check(ExtractResult == EGeometryFlowResult::Ok);
-
-	for (FBakeTextureGraphInfo& TexBakeStep : BakeTextureNodes)
-	{
-		TUniquePtr<UE::GeometryFlow::FTextureImage> NewImage = MakeUnique<UE::GeometryFlow::FTextureImage>();
-		bool bTakeTexBake = true;
-		ExtractResult = Exec.GetOutput(TexBakeStep.BakeNode,
-									   FBakeMeshTextureImageNode::OutParamTextureImage(),
-									   *NewImage,
-									   (int)EMeshProcessingDataTypes::TextureImage,
-									   bTakeNormalMap);
-		if (ExtractResult == EGeometryFlowResult::OperationCancelled)
-		{
-			return;
-		}
-		check(ExtractResult == EGeometryFlowResult::Ok);
-		TextureImages.Add(MoveTemp(NewImage));
-	}
-
-	ResultCollision = FSimpleShapeSet3d();
-	bool bTakeResultCollision = false;
-	ExtractResult = Exec.GetOutput(CollisionOutputNode,
-								   FCollisionGeometryTransferNode::OutParamValue(),
-								   ResultCollision,
-								   (int)FCollisionGeometry::DataTypeIdentifier, bTakeResultCollision);
-	if (ExtractResult == EGeometryFlowResult::OperationCancelled)
-	{
-		return;
-	}
-	check(ExtractResult == EGeometryFlowResult::Ok);
-
-	bool bTakeResultTangents = false;
-	ResultTangents = FMeshTangentsd();
-	ExtractResult = Exec.GetOutput(TangentsOutputNode, 
-								   FMeshTangentsTransferNode::OutParamValue(),
-								   ResultTangents, 
-								   (int)EMeshProcessingDataTypes::MeshTangentSet, 
-								   bTakeResultTangents);
-	if (ExtractResult == EGeometryFlowResult::OperationCancelled)
-	{
-		return;
-	}
-	check(ExtractResult == EGeometryFlowResult::Ok);
-	
-
-	bool bTakeResultMesh = true;
-	ResultMesh.Clear();
-	ExtractResult = Exec.GetOutput(MeshOutputNode, 
-								   FDynamicMeshTransferNode::OutParamValue(),
-								   ResultMesh, 
-								   (int32)EMeshProcessingDataTypes::DynamicMesh, 
-								   bTakeResultMesh);
-	if (ExtractResult == EGeometryFlowResult::OperationCancelled)
-	{
-		return;
-	}
-	check(ExtractResult == EGeometryFlowResult::Ok);
-}
-
-#endif
-
-
 void FGenerateMeshLODGraph::UpdatePreFilterSettings(const FMeshLODGraphPreFilterSettings& PreFilterSettings)
 {
 	UpdateSourceNodeValue<FNameSourceNode>(*Graph, FilterGroupsLayerNameNode, PreFilterSettings.FilterGroupLayerName);
diff --git a/Engine/Plugins/Experimental/MeshLODToolset/Source/MeshLODToolset/Private/Graphs/GenerateStaticMeshLODProcess.cpp b/Engine/Plugins/Experimental/MeshLODToolset/Source/MeshLODToolset/Private/Graphs/GenerateStaticMeshLODProcess.cpp
--- a/Engine/Plugins/Experimental/MeshLODToolset/Source/MeshLODToolset/Private/Graphs/GenerateStaticMeshLODProcess.cpp
+++ b/Engine/Plugins/Experimental/MeshLODToolset/Source/MeshLODToolset/Private/Graphs/GenerateStaticMeshLODProcess.cpp
@@ -515,26 +515,14 @@ void UGenerateStaticMeshLODProcess::UpdateSettings(const FGenerateStaticMeshLODP
 bool UGenerateStaticMeshLODProcess::ComputeDerivedSourceData(FProgressCancel* Progress)
 {
 	DerivedTextureImages.Reset();
-	if (bUseParallelExecutor)
-	{
-		Generator->EvaluateResultParallel(
-			this->DerivedLODMesh,
-			this->DerivedLODMeshTangents,
-			this->DerivedCollision,
-			this->DerivedNormalMapImage,
-			this->DerivedTextureImages,
-			Progress);
-	}
-	else
-	{
-		Generator->EvaluateResult(
-			this->DerivedLODMesh,
-			this->DerivedLODMeshTangents,
-			this->DerivedCollision,
-			this->DerivedNormalMapImage,
-			this->DerivedTextureImages,
-			Progress);
-	}
+
+	Generator->EvaluateResult(
+		this->DerivedLODMesh,
+		this->DerivedLODMeshTangents,
+		this->DerivedCollision,
+		this->DerivedNormalMapImage,
+		this->DerivedTextureImages,
+		Progress);
 
 	if (Progress && Progress->Cancelled())
 	{
diff --git a/Engine/Plugins/Experimental/MeshLODToolset/Source/MeshLODToolset/Private/Tools/GenerateStaticMeshLODAssetTool.cpp b/Engine/Plugins/Experimental/MeshLODToolset/Source/MeshLODToolset/Private/Tools/GenerateStaticMeshLODAssetTool.cpp
--- a/Engine/Plugins/Experimental/MeshLODToolset/Source/MeshLODToolset/Private/Tools/GenerateStaticMeshLODAssetTool.cpp
+++ b/Engine/Plugins/Experimental/MeshLODToolset/Source/MeshLODToolset/Private/Tools/GenerateStaticMeshLODAssetTool.cpp
@@ -302,16 +302,6 @@ void UGenerateStaticMeshLODAssetTool::Setup()
 
 	CollisionPreview = NewObject<UPreviewGeometry>(this);
 	CollisionPreview->CreateInWorld(TargetWorld, PreviewTransform);
-
-	// Recompute if we switch between parallel and serial
-	int32 WatcherIndex = BasicProperties->WatchProperty(BasicProperties->bParallelExecution, [this](bool bNewParallelExec)
-	{
-		GenerateProcess->bUseParallelExecutor = BasicProperties->bParallelExecution;
-		OnSettingsModified();
-	});
-
-	BasicProperties->SilentUpdateWatcherAtIndex(WatcherIndex);
-
 }
 
 void UGenerateStaticMeshLODAssetTool::OnSettingsModified()
diff --git a/Engine/Plugins/Experimental/MeshLODToolset/Source/MeshLODToolset/Public/Graphs/GenerateMeshLODGraph.h b/Engine/Plugins/Experimental/MeshLODToolset/Source/MeshLODToolset/Public/Graphs/GenerateMeshLODGraph.h
--- a/Engine/Plugins/Experimental/MeshLODToolset/Source/MeshLODToolset/Public/Graphs/GenerateMeshLODGraph.h
+++ b/Engine/Plugins/Experimental/MeshLODToolset/Source/MeshLODToolset/Public/Graphs/GenerateMeshLODGraph.h
@@ -71,14 +71,10 @@ public:
 		TArray<TUniquePtr<UE::GeometryFlow::FTextureImage>>& TextureImages,
 		FProgressCancel* Progress);
 
-	void EvaluateResultParallel(FDynamicMesh3& ResultMesh,
 								UE::Geometry::FMeshTangentsd& ResultTangents,
 								UE::Geometry::FSimpleShapeSet3d& ResultCollision,
-								UE::GeometryFlow::FNormalMapImage& NormalMap,
-								TArray<TUniquePtr<UE::GeometryFlow::FTextureImage>>& TextureImages,
-								FProgressCancel* Progress);
-
 protected:
+
 	TUniquePtr<UE::GeometryFlow::FGraph> Graph;
 
 
diff --git a/Engine/Plugins/Experimental/MeshLODToolset/Source/MeshLODToolset/Public/Graphs/GenerateStaticMeshLODProcess.h b/Engine/Plugins/Experimental/MeshLODToolset/Source/MeshLODToolset/Public/Graphs/GenerateStaticMeshLODProcess.h
--- a/Engine/Plugins/Experimental/MeshLODToolset/Source/MeshLODToolset/Public/Graphs/GenerateStaticMeshLODProcess.h
+++ b/Engine/Plugins/Experimental/MeshLODToolset/Source/MeshLODToolset/Public/Graphs/GenerateStaticMeshLODProcess.h
@@ -204,8 +204,6 @@ public:
 	};
 	void GetDerivedMaterialsPreview(FPreviewMaterials& MaterialSetOut);
 
-	bool bUseParallelExecutor = false;
-
 	FCriticalSection GraphEvalCriticalSection;
 
 protected:
diff --git a/Engine/Plugins/Experimental/MeshLODToolset/Source/MeshLODToolset/Public/Tools/GenerateStaticMeshLODAssetTool.h b/Engine/Plugins/Experimental/MeshLODToolset/Source/MeshLODToolset/Public/Tools/GenerateStaticMeshLODAssetTool.h
--- a/Engine/Plugins/Experimental/MeshLODToolset/Source/MeshLODToolset/Public/Tools/GenerateStaticMeshLODAssetTool.h
+++ b/Engine/Plugins/Experimental/MeshLODToolset/Source/MeshLODToolset/Public/Tools/GenerateStaticMeshLODAssetTool.h
@@ -76,10 +76,6 @@ public:
 	UPROPERTY(EditAnywhere, Category = AssetOptions, meta = (EditCondition = "OutputMode == EGenerateLODAssetOutputMode::UpdateExistingAsset"))
 	bool bSaveAsHDSource = true;
 
-
-	UPROPERTY(EditAnywhere, Category = Settings)
-	bool bParallelExecution = false;
-
 	UPROPERTY(EditAnywhere, Category = Settings)
 	FGenerateStaticMeshLODProcessSettings GeneratorSettings;
 
