Fix compile error from bad robomerge in GenerateMeshLODGraph.h

diff --git a/Engine/Plugins/Experimental/MeshLODToolset/Source/MeshLODToolset/Public/Graphs/GenerateMeshLODGraph.h b/Engine/Plugins/Experimental/MeshLODToolset/Source/MeshLODToolset/Public/Graphs/GenerateMeshLODGraph.h
--- a/Engine/Plugins/Experimental/MeshLODToolset/Source/MeshLODToolset/Public/Graphs/GenerateMeshLODGraph.h
+++ b/Engine/Plugins/Experimental/MeshLODToolset/Source/MeshLODToolset/Public/Graphs/GenerateMeshLODGraph.h
@@ -71,8 +71,6 @@ public:
 		TArray<TUniquePtr<UE::GeometryFlow::FTextureImage>>& TextureImages,
 		FProgressCancel* Progress);
 
-								UE::Geometry::FMeshTangentsd& ResultTangents,
-								UE::Geometry::FSimpleShapeSet3d& ResultCollision,
 protected:
 
 	TUniquePtr<UE::GeometryFlow::FGraph> Graph;
