- Add warning/error messages to FProgressCancel
- Pop up a SlowTask modal progress bar when starting AutoLOD tool, just to avoid locking up the entire UI

#rnx

diff --git a/Engine/Plugins/Experimental/GeometryProcessing/Source/GeometricObjects/Public/Util/ProgressCancel.h b/Engine/Plugins/Experimental/GeometryProcessing/Source/GeometricObjects/Public/Util/ProgressCancel.h
--- a/Engine/Plugins/Experimental/GeometryProcessing/Source/GeometricObjects/Public/Util/ProgressCancel.h
+++ b/Engine/Plugins/Experimental/GeometryProcessing/Source/GeometricObjects/Public/Util/ProgressCancel.h
@@ -32,4 +32,37 @@ public:
 		WasCancelled = CancelF();
 		return WasCancelled;
 	}
+
+
+public:
+
+	enum class EMessageLevel
+	{
+		// Note: Corresponds to EToolMessageLevel in InteractiveToolsFramework/ToolContextInterfaces.h
+
+		/** Development message goes into development log */
+		Internal = 0,
+		/** User message should appear in user-facing log */
+		UserMessage = 1,
+		/** Notification message should be shown in a non-modal notification window */
+		UserNotification = 2,
+		/** Warning message should be shown in a non-modal notification window with panache */
+		UserWarning = 3,
+		/** Error message should be shown in a modal notification window */
+		UserError = 4
+	};
+
+	struct FMessageInfo
+	{
+		FText MessageText;
+		EMessageLevel MessageLevel;
+		FDateTime Timestamp;
+	};
+
+	void AddWarning(const FText& MessageText, EMessageLevel MessageLevel)
+	{
+		Warnings.Add(FMessageInfo{ MessageText , MessageLevel, FDateTime::Now() });
+	}
+
+	TArray<FMessageInfo> Warnings;
 };
diff --git a/Engine/Plugins/Experimental/MeshLODToolset/Source/MeshLODToolset/Private/Graphs/GenerateStaticMeshLODProcess.cpp b/Engine/Plugins/Experimental/MeshLODToolset/Source/MeshLODToolset/Private/Graphs/GenerateStaticMeshLODProcess.cpp
--- a/Engine/Plugins/Experimental/MeshLODToolset/Source/MeshLODToolset/Private/Graphs/GenerateStaticMeshLODProcess.cpp
+++ b/Engine/Plugins/Experimental/MeshLODToolset/Source/MeshLODToolset/Private/Graphs/GenerateStaticMeshLODProcess.cpp
@@ -95,7 +95,7 @@ namespace GenerateStaticMeshLODProcessHelpers
 	}
 }
 
-bool UGenerateStaticMeshLODProcess::Initialize(UStaticMesh* StaticMeshIn)
+bool UGenerateStaticMeshLODProcess::Initialize(UStaticMesh* StaticMeshIn, FProgressCancel* Progress)
 {
 	if (!ensure(StaticMeshIn)) return false;
 	if (!ensure(StaticMeshIn->GetNumSourceModels() > 0)) return false;
@@ -172,10 +172,15 @@ bool UGenerateStaticMeshLODProcess::Initialize(UStaticMesh* StaticMeshIn)
 			}
 			if (bFoundTextureNonParamExpession)
 			{
-				// TODO: bubble this up to the tool level
-				UE_LOG(LogTemp, Warning, 
-					   TEXT("UGenerateStaticMeshLODProcess: Non-parameter texture sampler detected in input material [%s]. Output materials may have unexpected behaviour."),
-					   *Material->GetName());
+				FString WarningMessage("WARNING: UGenerateStaticMeshLODProcess: Non-parameter texture sampler detected in input material [%s]. Output materials may have unexpected behaviour.");
+				UE_LOG(LogTemp, Warning, TEXT("%s"), *WarningMessage, *Material->GetName());
+
+				if (Progress)
+				{
+					FText WarningText = FText::Format(LOCTEXT("NonParameterTextureWarning", "Non-parameter texture sampler detected in input material [{0}]. Output materials may have unexpected behaviour."),
+								  FText::FromString(Material->GetName()));
+					Progress->AddWarning(WarningText, FProgressCancel::EMessageLevel::UserWarning);
+				}
 			}
 		}
 
diff --git a/Engine/Plugins/Experimental/MeshLODToolset/Source/MeshLODToolset/Private/Tools/GenerateStaticMeshLODAssetTool.cpp b/Engine/Plugins/Experimental/MeshLODToolset/Source/MeshLODToolset/Private/Tools/GenerateStaticMeshLODAssetTool.cpp
--- a/Engine/Plugins/Experimental/MeshLODToolset/Source/MeshLODToolset/Private/Tools/GenerateStaticMeshLODAssetTool.cpp
+++ b/Engine/Plugins/Experimental/MeshLODToolset/Source/MeshLODToolset/Private/Tools/GenerateStaticMeshLODAssetTool.cpp
@@ -47,9 +47,8 @@
 #include "TargetInterfaces/StaticMeshBackedTarget.h"
 #include "ToolTargetManager.h"
 
-#if WITH_EDITOR
+static_assert(WITH_EDITOR, "Tool being compiled without editor");
 #include "Misc/ScopedSlowTask.h"
-#endif
 
 #include "ExplicitUseGeometryMathTypes.h"		// using UE::Geometry::(math types)
 using namespace UE::Geometry;
@@ -207,11 +206,33 @@ void UGenerateStaticMeshLODAssetTool::Setup()
 	IPrimitiveComponentBackedTarget* SourceComponent = TargetComponentInterface(0);
 	UStaticMeshComponent* StaticMeshComponent = CastChecked<UStaticMeshComponent>(SourceComponent->GetOwnerComponent());
 	UStaticMesh* StaticMesh = StaticMeshComponent->GetStaticMesh();
+
+	FProgressCancel Progress;
+	FScopedSlowTask SlowTask(2, LOCTEXT("UGenerateStaticMeshLODAssetTool_Setup", "Initializing tool ..."));
+	SlowTask.MakeDialog();
+
 	if (StaticMesh)
 	{
-		GenerateProcess->Initialize(StaticMesh);		// Must happen on main thread
+		SlowTask.EnterProgressFrame(1);
+
+		bool bInitializeOK = GenerateProcess->Initialize(StaticMesh, &Progress);		// Must happen on main thread
+
+		if (Progress.Warnings.Num() > 0)
+		{
+			const FProgressCancel::FMessageInfo& Warning = Progress.Warnings[0];
+			GetToolManager()->DisplayMessage(Warning.MessageText, (EToolMessageLevel)Warning.MessageLevel);
+		}
+	}
+	else
+	{
+		GetToolManager()->DisplayMessage(
+			LOCTEXT("GenerateStaticMeshLODAssetTool_NoStaticMesh", "Could not find Static Mesh in selected input"),
+			EToolMessageLevel::UserError);
 	}
 
+	
+
+	SlowTask.EnterProgressFrame(1);
 	BasicProperties = NewObject<UGenerateStaticMeshLODAssetToolProperties>(this);
 	AddToolPropertySource(BasicProperties);
 	BasicProperties->RestoreProperties(this);
diff --git a/Engine/Plugins/Experimental/MeshLODToolset/Source/MeshLODToolset/Public/Graphs/GenerateStaticMeshLODProcess.h b/Engine/Plugins/Experimental/MeshLODToolset/Source/MeshLODToolset/Public/Graphs/GenerateStaticMeshLODProcess.h
--- a/Engine/Plugins/Experimental/MeshLODToolset/Source/MeshLODToolset/Public/Graphs/GenerateStaticMeshLODProcess.h
+++ b/Engine/Plugins/Experimental/MeshLODToolset/Source/MeshLODToolset/Public/Graphs/GenerateStaticMeshLODProcess.h
@@ -163,7 +163,7 @@ class UGenerateStaticMeshLODProcess : public UObject
 	GENERATED_BODY()
 public:
 
-	bool Initialize(UStaticMesh* SourceMesh);
+	bool Initialize(UStaticMesh* SourceMesh, FProgressCancel* Progress = nullptr);
 
 	const FGenerateStaticMeshLODProcessSettings& GetCurrentSettings() const { return CurrentSettings; }
 	void UpdateSettings(const FGenerateStaticMeshLODProcessSettings& NewSettings);
