GeometryProcessing: support initializing TMeshTangents from FDynamicMesh3 tangents attribute sets.
LODToolset: use correct tangents on HiRes, SourceModel and RenderData LODs
#rnx

diff --git a/Engine/Plugins/Experimental/GeometryProcessing/Source/DynamicMesh/Private/MeshTangents.cpp b/Engine/Plugins/Experimental/GeometryProcessing/Source/DynamicMesh/Private/MeshTangents.cpp
--- a/Engine/Plugins/Experimental/GeometryProcessing/Source/DynamicMesh/Private/MeshTangents.cpp
+++ b/Engine/Plugins/Experimental/GeometryProcessing/Source/DynamicMesh/Private/MeshTangents.cpp
@@ -28,6 +28,41 @@ void TMeshTangents<RealType>::SetTangentCount(int Count, bool bClearToZero)
 }
 
 
+template<typename RealType>
+bool TMeshTangents<RealType>::CopyTriVertexTangents(const FDynamicMesh3& SourceMesh)
+{
+	if (SourceMesh.HasAttributes() == false || SourceMesh.Attributes()->HasTangentSpace() == false)
+	{
+		return false;
+	}
+	const FDynamicMeshNormalOverlay* TangentsAttrib = SourceMesh.Attributes()->PrimaryTangents();
+	const FDynamicMeshNormalOverlay* BiTangentsAttrib = SourceMesh.Attributes()->PrimaryBiTangents();
+
+	//InitializeTriVertexTangents(false);    // this requires Mesh to be initialized which is not strictly necessary...
+	SetTangentCount(SourceMesh.MaxTriangleID() * 3, false);
+
+	for (int32 tid : SourceMesh.TriangleIndicesItr())
+	{
+		if ( TangentsAttrib->IsSetTriangle(tid) && BiTangentsAttrib->IsSetTriangle(tid) )
+		{
+			FVector3f TriTangents[3];
+			TangentsAttrib->GetTriElements(tid, TriTangents[0], TriTangents[1], TriTangents[2]);
+			FVector3f TriBiTangents[3];
+			BiTangentsAttrib->GetTriElements(tid, TriBiTangents[0], TriBiTangents[1], TriBiTangents[2]);
+
+			int32 k = 3 * tid;
+			for (int32 j = 0; j < 3; ++j )
+			{
+				Tangents[k+j] = FVector3<RealType>(TriTangents[j]);
+				Bitangents[k+j] = FVector3<RealType>(TriBiTangents[j]);
+			}
+		}
+	}
+
+	return true;
+}
+
+
 
 template<typename RealType>
 void TMeshTangents<RealType>::ComputeTriVertexTangents(const FDynamicMeshNormalOverlay* NormalOverlay, const FDynamicMeshUVOverlay* UVOverlay, const FComputeTangentsOptions& Options)
diff --git a/Engine/Plugins/Experimental/GeometryProcessing/Source/DynamicMesh/Public/MeshTangents.h b/Engine/Plugins/Experimental/GeometryProcessing/Source/DynamicMesh/Public/MeshTangents.h
--- a/Engine/Plugins/Experimental/GeometryProcessing/Source/DynamicMesh/Public/MeshTangents.h
+++ b/Engine/Plugins/Experimental/GeometryProcessing/Source/DynamicMesh/Public/MeshTangents.h
@@ -84,6 +84,13 @@ public:
 	template<typename OtherRealType>
 	void CopyTriVertexTangents(const TMeshTangents<OtherRealType>& OtherTangents);
 
+	/**
+	 * Initialize Tangents from the FDynamicMeshAttributeSet of SourceMesh 
+	 * @return false if tangent attributes do not exist
+	 */
+	bool CopyTriVertexTangents(const FDynamicMesh3& SourceMesh);
+
+
 	/**
 	 * Calculate per-triangle tangent spaces based on the given per-triangle normal and UV overlays.
 	 * In this mode there is no averaging of tangents across triangles. So if we have N triangles
diff --git a/Engine/Plugins/Experimental/GeometryProcessing/Source/MeshConversion/Private/IndexedMeshToDynamicMesh.cpp b/Engine/Plugins/Experimental/GeometryProcessing/Source/MeshConversion/Private/IndexedMeshToDynamicMesh.cpp
--- a/Engine/Plugins/Experimental/GeometryProcessing/Source/MeshConversion/Private/IndexedMeshToDynamicMesh.cpp
+++ b/Engine/Plugins/Experimental/GeometryProcessing/Source/MeshConversion/Private/IndexedMeshToDynamicMesh.cpp
@@ -77,6 +77,8 @@ void UE::Conversion::RenderBuffersToDynamicMesh(
 		int32 tid = MeshOut.AppendTriangle(a, b, c);
 		ensure(tid == ti);
 		Normals->SetTriangle(tid, FIndex3i(a, b, c));
+		Tangents->SetTriangle(tid, FIndex3i(a, b, c));
+		BiTangents->SetTriangle(tid, FIndex3i(a, b, c));
 		for (int32 ui = 0; ui < NumUVChannels; ++ui)
 		{
 			UVs[ui]->SetTriangle(tid, FIndex3i(a, b, c));
diff --git a/Engine/Plugins/Experimental/MeshLODToolset/Source/MeshLODToolset/Private/Tools/LODManagerTool.cpp b/Engine/Plugins/Experimental/MeshLODToolset/Source/MeshLODToolset/Private/Tools/LODManagerTool.cpp
--- a/Engine/Plugins/Experimental/MeshLODToolset/Source/MeshLODToolset/Private/Tools/LODManagerTool.cpp
+++ b/Engine/Plugins/Experimental/MeshLODToolset/Source/MeshLODToolset/Private/Tools/LODManagerTool.cpp
@@ -108,6 +108,7 @@ void ULODManagerTool::Setup()
 
 		LODPreview = NewObject<UPreviewMesh>(this);
 		LODPreview->CreateInWorld(TargetWorld, TargetComponentInterface(0)->GetWorldTransform());
+		LODPreview->SetTangentsMode(EDynamicMeshTangentCalcType::ExternallyCalculated);
 		LODPreview->SetVisible(false);
 
 		LODPreviewLines = NewObject<UPreviewGeometry>(this);
@@ -351,6 +352,7 @@ void ULODManagerTool::UpdatePreviewLines(FLODMeshInfo& MeshInfo)
 	}
 	if (BoundaryLines)
 	{
+		BoundaryLines->Clear();
 		for (int32 eid : MeshInfo.BoundaryEdges)
 		{
 			FVector3d A, B;
@@ -392,6 +394,7 @@ void ULODManagerTool::UpdatePreviewLOD()
 	{
 		TargetComponentInterface(0)->SetOwnerVisibility(false);
 		LODPreview->ReplaceMesh( (*Found)->Mesh );
+		LODPreview->UpdateTangents<float>( &(*Found)->Tangents, false);
 		LODPreview->SetVisible(true);
 		UpdatePreviewLines(**Found);
 		LODPreviewLines->SetAllVisible(true);
@@ -419,8 +422,14 @@ bool ULODManagerTool::CacheLODMesh(const FString& Name, FLODName LODName)
 		FMeshDescription* HiResMeshDescription = StaticMesh->GetHiResMeshDescription();
 		if (HiResMeshDescription != nullptr)
 		{
+			// TODO: we only need to do this copy if normals or tangents are set to auto-generated...
+			FMeshBuildSettings HiResBuildSettings = StaticMesh->GetHiResSourceModel().BuildSettings;
+			FMeshDescription TmpMeshDescription(*HiResMeshDescription);
+			UE::MeshDescription::InitializeAutoGeneratedAttributes(TmpMeshDescription, &HiResBuildSettings);
+
 			TUniquePtr<FLODMeshInfo> NewMeshInfo = MakeUnique<FLODMeshInfo>();
-			Converter.Convert(HiResMeshDescription, NewMeshInfo->Mesh);
+			Converter.Convert(&TmpMeshDescription, NewMeshInfo->Mesh);
+			Converter.CopyTangents(&TmpMeshDescription, &NewMeshInfo->Mesh, &NewMeshInfo->Tangents);
 			LODMeshCache.Add(Name, MoveTemp(NewMeshInfo));
 			return true;
 		}
@@ -431,8 +440,14 @@ bool ULODManagerTool::CacheLODMesh(const FString& Name, FLODName LODName)
 		FMeshDescription* LODMeshDescription = StaticMesh->GetMeshDescription(LODName.SourceModelIndex);
 		if (LODMeshDescription != nullptr)
 		{
+			// TODO: we only need to do this copy if normals or tangents are set to auto-generated...
+			FMeshBuildSettings LODBuildSettings = StaticMesh->GetSourceModel(LODName.SourceModelIndex).BuildSettings;
+			FMeshDescription TmpMeshDescription(*LODMeshDescription);
+			UE::MeshDescription::InitializeAutoGeneratedAttributes(TmpMeshDescription, &LODBuildSettings);
+
 			TUniquePtr<FLODMeshInfo> NewMeshInfo = MakeUnique<FLODMeshInfo>();
-			Converter.Convert(LODMeshDescription, NewMeshInfo->Mesh);
+			Converter.Convert(&TmpMeshDescription, NewMeshInfo->Mesh);
+			Converter.CopyTangents(&TmpMeshDescription, &NewMeshInfo->Mesh, &NewMeshInfo->Tangents);
 			LODMeshCache.Add(Name, MoveTemp(NewMeshInfo));
 			return true;
 		}
@@ -447,6 +462,8 @@ bool ULODManagerTool::CacheLODMesh(const FString& Name, FLODName LODName)
 				const FStaticMeshLODResources& LODResource = RenderData->LODResources[LODName.RenderDataIndex];
 				TUniquePtr<FLODMeshInfo> NewMeshInfo = MakeUnique<FLODMeshInfo>();
 				UE::Conversion::RenderBuffersToDynamicMesh(LODResource.VertexBuffers, LODResource.IndexBuffer, LODResource.Sections, NewMeshInfo->Mesh);
+				NewMeshInfo->Tangents.SetMesh(&NewMeshInfo->Mesh);
+				NewMeshInfo->Tangents.CopyTriVertexTangents(NewMeshInfo->Mesh);
 				LODMeshCache.Add(Name, MoveTemp(NewMeshInfo));
 				return true;
 			}
diff --git a/Engine/Plugins/Experimental/MeshLODToolset/Source/MeshLODToolset/Public/Tools/LODManagerTool.h b/Engine/Plugins/Experimental/MeshLODToolset/Source/MeshLODToolset/Public/Tools/LODManagerTool.h
--- a/Engine/Plugins/Experimental/MeshLODToolset/Source/MeshLODToolset/Public/Tools/LODManagerTool.h
+++ b/Engine/Plugins/Experimental/MeshLODToolset/Source/MeshLODToolset/Public/Tools/LODManagerTool.h
@@ -8,6 +8,7 @@
 #include "Engine/StaticMesh.h"
 #include "PreviewMesh.h"
 #include "Drawing/PreviewGeometryActor.h"
+#include "MeshTangents.h"
 #include "LODManagerTool.generated.h"
 
 // predeclarations
@@ -234,6 +235,7 @@ protected:
 	struct FLODMeshInfo
 	{
 		UE::Geometry::FDynamicMesh3 Mesh;
+		UE::Geometry::FMeshTangentsf Tangents;		// todo: store tangents on the Mesh. Requires support from SimpleDynamicMeshComponent.
 
 		bool bInfoCached = false;
 		TArray<int> BoundaryEdges;
