Fix for AutoLOD not handling multiple texture parameters pointing to the same texture.

#rnx

diff --git a/Engine/Plugins/Experimental/MeshLODToolset/Source/MeshLODToolset/Private/Graphs/GenerateStaticMeshLODProcess.cpp b/Engine/Plugins/Experimental/MeshLODToolset/Source/MeshLODToolset/Private/Graphs/GenerateStaticMeshLODProcess.cpp
--- a/Engine/Plugins/Experimental/MeshLODToolset/Source/MeshLODToolset/Private/Graphs/GenerateStaticMeshLODProcess.cpp
+++ b/Engine/Plugins/Experimental/MeshLODToolset/Source/MeshLODToolset/Private/Graphs/GenerateStaticMeshLODProcess.cpp
@@ -734,7 +734,7 @@ void UGenerateStaticMeshLODProcess::WriteDerivedTextures(bool bCreatingNewStatic
 	// this is a workaround for handling multiple materials that reference the same texture. Currently the code
 	// below will try to write that texture multiple times, which will fail when it tries to create a package
 	// for a filename that already exists
-	TArray<UTexture2D*> SourceTexturesWritten;
+	TMap<UTexture2D*, UTexture2D*> WrittenSourceToDerived;
 
 	// write derived textures
 	int32 NumMaterials = SourceMaterials.Num();
@@ -754,13 +754,16 @@ void UGenerateStaticMeshLODProcess::WriteDerivedTextures(bool bCreatingNewStatic
 		for (int32 ti = 0; ti < NumTextures; ++ti)
 		{
 			const FTextureInfo& SourceTex = SourceMaterialInfo.SourceTextures[ti];
-			if (SourceTexturesWritten.Contains(SourceTex.Texture))
+			FTextureInfo& DerivedTex = DerivedMaterialInfo.DerivedTextures[ti];
+
+			if (WrittenSourceToDerived.Contains(SourceTex.Texture))
 			{
+				// Already computed the derived texture from this source
+				DerivedTex.Texture = WrittenSourceToDerived[SourceTex.Texture];
 				continue;
 			}
 
 			bool bConvertToSRGB = SourceTex.Texture->SRGB;
-			FTextureInfo& DerivedTex = DerivedMaterialInfo.DerivedTextures[ti];
 
 			if (DerivedTex.bShouldBakeTexture == false)
 			{
@@ -784,7 +787,7 @@ void UGenerateStaticMeshLODProcess::WriteDerivedTextures(bool bCreatingNewStatic
 				bool bWriteOK = WriteDerivedTexture(SourceTex.Texture, DerivedTex.Texture, bCreatingNewStaticMeshAsset);
 				ensure(bWriteOK);
 
-				SourceTexturesWritten.Add(SourceTex.Texture);
+				WrittenSourceToDerived.Add(SourceTex.Texture, DerivedTex.Texture);
 			}
 		}
 
