Fix crash while running AutoLOD on a mesh containing an asset with the same name as an auto-generated new asset. Warn the user if they have an unreference material in the input mesh.

#rnx

diff --git a/Engine/Plugins/Experimental/MeshLODToolset/Source/MeshLODToolset/Private/Graphs/GenerateStaticMeshLODProcess.cpp b/Engine/Plugins/Experimental/MeshLODToolset/Source/MeshLODToolset/Private/Graphs/GenerateStaticMeshLODProcess.cpp
--- a/Engine/Plugins/Experimental/MeshLODToolset/Source/MeshLODToolset/Private/Graphs/GenerateStaticMeshLODProcess.cpp
+++ b/Engine/Plugins/Experimental/MeshLODToolset/Source/MeshLODToolset/Private/Graphs/GenerateStaticMeshLODProcess.cpp
@@ -93,6 +93,45 @@ namespace GenerateStaticMeshLODProcessHelpers
 			}
 		}
 	}
+
+	void FindUnreferencedMaterials(UStaticMesh* StaticMesh, FMeshDescription* MeshDescription, TArray<int32>& OutUnreferencedMaterials)
+	{
+		const TArray<FStaticMaterial>& CurMaterialSet = StaticMesh->GetStaticMaterials();
+		const int32 NumMaterials = CurMaterialSet.Num();
+		TArray<bool> MatUsedFlags;
+		MatUsedFlags.Init(false, NumMaterials);
+
+		TArray<int32> MeshMaterialIndices;
+		TMap<FPolygonGroupID, int32> PolygonGroupToMaterialIndex;
+		const FStaticMeshConstAttributes Attributes(*MeshDescription);
+		TPolygonGroupAttributesConstRef<FName> PolygonGroupImportedMaterialSlotNames = Attributes.GetPolygonGroupMaterialSlotNames();
+
+		for (FPolygonGroupID PolygonGroupID : MeshDescription->PolygonGroups().GetElementIDs())
+		{
+			int32 MaterialIndex = StaticMesh->GetMaterialIndexFromImportedMaterialSlotName(PolygonGroupImportedMaterialSlotNames[PolygonGroupID]);
+			if (MaterialIndex == INDEX_NONE)
+			{
+				MaterialIndex = PolygonGroupID.GetValue();
+			}
+			PolygonGroupToMaterialIndex.Add(PolygonGroupID, MaterialIndex);
+		}
+
+		for (const FTriangleID TriangleID : MeshDescription->Triangles().GetElementIDs())
+		{
+			const FPolygonGroupID PolygonGroupID = MeshDescription->GetTrianglePolygonGroup(TriangleID);
+			const int32 MaterialIndex = PolygonGroupToMaterialIndex[PolygonGroupID];
+			MatUsedFlags[MaterialIndex] = true;
+			MeshMaterialIndices.AddUnique(MaterialIndex);
+		}
+
+		for (int32 MaterialID = 0; MaterialID < NumMaterials; ++MaterialID)
+		{
+			if (!MatUsedFlags[MaterialID])
+			{
+				OutUnreferencedMaterials.Emplace(MaterialID);
+			}
+		}
+	}
 }
 
 bool UGenerateStaticMeshLODProcess::Initialize(UStaticMesh* StaticMeshIn, FProgressCancel* Progress)
@@ -127,6 +166,24 @@ bool UGenerateStaticMeshLODProcess::Initialize(UStaticMesh* StaticMeshIn, FProgr
 
 	// get list of source materials and find all the input texture params
 	const TArray<FStaticMaterial>& Materials = SourceStaticMesh->GetStaticMaterials();
+
+	// warn the user if there are any unsed materials in the mesh
+	if (Progress)
+	{
+		TArray<int32> UnusedMaterials;
+		GenerateStaticMeshLODProcessHelpers::FindUnreferencedMaterials(SourceStaticMesh, SourceMeshDescription.Get(), UnusedMaterials);
+		if (UnusedMaterials.Num() > 0)
+		{
+			for (int32 UnusedMaterialIndex : UnusedMaterials)
+			{
+				FText WarningText = FText::Format(LOCTEXT("UnusedMaterialWarning", "Found an unused material ({0}). Consider removing it before using this tool."),
+												  FText::FromName(Materials[UnusedMaterialIndex].MaterialInterface->GetFName()));
+				Progress->AddWarning(WarningText, FProgressCancel::EMessageLevel::UserWarning);
+			}
+		}
+	}
+
+
 	SourceMaterials.SetNum(Materials.Num());
 	TArray<FReadTextureJob> ReadJobs;
 	for (int32 mi = 0; mi < Materials.Num(); ++mi)
@@ -727,6 +784,40 @@ void UGenerateStaticMeshLODProcess::UpdateSourceAsset(bool bSetNewHDSourceAsset)
 }
 
 
+
+bool UGenerateStaticMeshLODProcess::IsSourceAsset(const FString& AssetPath) const
+{
+	FAssetData AssetData = UEditorAssetLibrary::FindAssetData(AssetPath);
+
+	for (const FSourceMaterialInfo& MaterialInfo : SourceMaterials)
+	{
+		UMaterialInterface* MaterialInterface = MaterialInfo.SourceMaterial.MaterialInterface;
+		if (MaterialInterface == nullptr)
+		{
+			continue;
+		}
+
+		FString SourceMaterialPath = UEditorAssetLibrary::GetPathNameForLoadedAsset(MaterialInterface);
+		if (UEditorAssetLibrary::FindAssetData(SourceMaterialPath) == AssetData)
+		{
+			return true;
+		}
+
+		for (const FTextureInfo& TextureInfo : MaterialInfo.SourceTextures)
+		{
+			FString SourceTexturePath = UEditorAssetLibrary::GetPathNameForLoadedAsset(TextureInfo.Texture);
+			if (UEditorAssetLibrary::FindAssetData(SourceTexturePath) == AssetData)
+			{
+				return true;
+			}
+		}
+	}
+
+	return false;
+}
+
+
+
 void UGenerateStaticMeshLODProcess::WriteDerivedTextures(bool bCreatingNewStaticMeshAsset)
 {
 	IAssetTools& AssetTools = FModuleManager::Get().LoadModuleChecked<FAssetToolsModule>("AssetTools").Get();
@@ -835,7 +926,9 @@ bool UGenerateStaticMeshLODProcess::WriteDerivedTexture(UTexture2D* DerivedTextu
 	FString NewTexName = FString::Printf(TEXT("%s%s"), *BaseTexName, *DerivedSuffix);
 	FString NewAssetPath = FPaths::Combine(DerivedAssetFolder, NewTexName);
 
-	if (bCreatingNewStaticMeshAsset)
+	bool bNewAssetExistsInMemory = IsSourceAsset(NewAssetPath);
+
+	if (bCreatingNewStaticMeshAsset || bNewAssetExistsInMemory)
 	{
 		// Don't delete an existing asset. If name collision occurs, rename the new asset.
 		bool bNewAssetExists = UEditorAssetLibrary::DoesAssetExist(NewAssetPath);
@@ -908,8 +1001,9 @@ void UGenerateStaticMeshLODProcess::WriteDerivedMaterials(bool bCreatingNewStati
 		FString MaterialName = FPaths::GetBaseFilename(SourceMaterialPath, true);
 		FString NewMaterialName = FString::Printf(TEXT("%s%s"), *MaterialName, *DerivedSuffix);
 		FString NewMaterialPath = FPaths::Combine(DerivedAssetFolder, NewMaterialName);
+		bool bNewAssetExistsInMemory = IsSourceAsset(NewMaterialPath);
 
-		if (bCreatingNewStaticMeshAsset)
+		if (bCreatingNewStaticMeshAsset || bNewAssetExistsInMemory)
 		{
 			// Don't delete an existing material. If name collision occurs, rename the new material.
 			bool bNewAssetExists = UEditorAssetLibrary::DoesAssetExist(NewMaterialPath);
diff --git a/Engine/Plugins/Experimental/MeshLODToolset/Source/MeshLODToolset/Public/Graphs/GenerateStaticMeshLODProcess.h b/Engine/Plugins/Experimental/MeshLODToolset/Source/MeshLODToolset/Public/Graphs/GenerateStaticMeshLODProcess.h
--- a/Engine/Plugins/Experimental/MeshLODToolset/Source/MeshLODToolset/Public/Graphs/GenerateStaticMeshLODProcess.h
+++ b/Engine/Plugins/Experimental/MeshLODToolset/Source/MeshLODToolset/Public/Graphs/GenerateStaticMeshLODProcess.h
@@ -301,4 +301,7 @@ protected:
 
 	void UpdateMaterialTextureParameters(UMaterialInstanceDynamic* Material, const FSourceMaterialInfo& SourceMaterialInfo,
 		const TMap<UTexture2D*,UTexture2D*>& PreviewTextures, UTexture2D* PreviewNormalMap);
+
+	// Return true if the given path corresponds to a material or texture in SourceMaterials
+	bool IsSourceAsset(const FString& AssetPath) const;
 };
