- Added tool tips to AutoLOD tool properties.
- Changed a few tool property names.
- Hooked up the NewAssetName property so that it's used when saving a new mesh.

#rnx

diff --git a/Engine/Plugins/Experimental/MeshLODToolset/Source/MeshLODToolset/Private/Graphs/GenerateStaticMeshLODProcess.cpp b/Engine/Plugins/Experimental/MeshLODToolset/Source/MeshLODToolset/Private/Graphs/GenerateStaticMeshLODProcess.cpp
--- a/Engine/Plugins/Experimental/MeshLODToolset/Source/MeshLODToolset/Private/Graphs/GenerateStaticMeshLODProcess.cpp
+++ b/Engine/Plugins/Experimental/MeshLODToolset/Source/MeshLODToolset/Private/Graphs/GenerateStaticMeshLODProcess.cpp
@@ -381,7 +381,7 @@ bool UGenerateStaticMeshLODProcess::Initialize(UStaticMesh* StaticMeshIn, FProgr
 	SourceAssetFolder = FPaths::GetPath(SourceAssetPath);
 	SourceAssetName = FPaths::GetBaseFilename(FullPathWithExtension, true);
 
-	CalculateDerivedPathName(GetDefaultDerivedAssetSuffix());
+	CalculateDerivedPathName(SourceAssetName, GetDefaultDerivedAssetSuffix());
 
 	InitializeGenerator();
 
@@ -389,17 +389,23 @@ bool UGenerateStaticMeshLODProcess::Initialize(UStaticMesh* StaticMeshIn, FProgr
 }
 
 
-
-
-void UGenerateStaticMeshLODProcess::CalculateDerivedPathName(FString NewAssetSuffix)
+void UGenerateStaticMeshLODProcess::CalculateDerivedPathName(const FString& NewAssetBaseName, const FString& NewAssetSuffix)
 {
+	DerivedAssetNameNoSuffix = FPaths::MakeValidFileName(NewAssetBaseName);
+	if (DerivedAssetNameNoSuffix.Len() == 0)
+	{
+		DerivedAssetNameNoSuffix = SourceAssetName;
+	}
+
 	DerivedSuffix = FPaths::MakeValidFileName(NewAssetSuffix);
 	if (DerivedSuffix.Len() == 0)
 	{
 		DerivedSuffix = GetDefaultDerivedAssetSuffix();
 	}
-	DerivedAssetPath = FString::Printf(TEXT("%s%s"), *GetSourceAssetPath(), *DerivedSuffix);
-	DerivedAssetFolder = FPaths::GetPath(DerivedAssetPath);
+
+	DerivedAssetName = FString::Printf(TEXT("%s%s"), *DerivedAssetNameNoSuffix, *DerivedSuffix);
+	DerivedAssetFolder = SourceAssetFolder;
+	DerivedAssetPath = FPaths::Combine(DerivedAssetFolder, DerivedAssetName);
 }
 
 
@@ -898,7 +904,7 @@ void UGenerateStaticMeshLODProcess::WriteDerivedTextures(bool bCreatingNewStatic
 			FTexture2DBuilder::CopyPlatformDataToSourceData(DerivedNormalMapTex, FTexture2DBuilder::ETextureType::NormalMap);
 
 			// write asset
-			bool bWriteOK = WriteDerivedTexture(DerivedNormalMapTex, SourceAssetName + TEXT("_NormalMap"), bCreatingNewStaticMeshAsset);
+			bool bWriteOK = WriteDerivedTexture(DerivedNormalMapTex, DerivedAssetNameNoSuffix + TEXT("_NormalMap"), bCreatingNewStaticMeshAsset);
 			ensure(bWriteOK);
 		}
 	}
diff --git a/Engine/Plugins/Experimental/MeshLODToolset/Source/MeshLODToolset/Private/Tools/GenerateStaticMeshLODAssetTool.cpp b/Engine/Plugins/Experimental/MeshLODToolset/Source/MeshLODToolset/Private/Tools/GenerateStaticMeshLODAssetTool.cpp
--- a/Engine/Plugins/Experimental/MeshLODToolset/Source/MeshLODToolset/Private/Tools/GenerateStaticMeshLODAssetTool.cpp
+++ b/Engine/Plugins/Experimental/MeshLODToolset/Source/MeshLODToolset/Private/Tools/GenerateStaticMeshLODAssetTool.cpp
@@ -236,7 +236,8 @@ void UGenerateStaticMeshLODAssetTool::Setup()
 	BasicProperties = NewObject<UGenerateStaticMeshLODAssetToolProperties>(this);
 	AddToolPropertySource(BasicProperties);
 	BasicProperties->RestoreProperties(this);
-	BasicProperties->OutputName = AssetGenerationUtil::GetComponentAssetBaseName(TargetComponentInterface(0)->GetOwnerComponent());
+	FString FullPathWithExtension = UEditorAssetLibrary::GetPathNameForLoadedAsset(StaticMesh);
+	BasicProperties->NewAssetName = FPaths::GetBaseFilename(FullPathWithExtension, true);
 	BasicProperties->GeneratedSuffix = TEXT("_AutoLOD");
 	BasicProperties->GeneratorSettings = GenerateProcess->GetCurrentSettings();
 	BasicProperties->WatchProperty(BasicProperties->GeneratorSettings.FilterGroupLayer, [this](FName) { OnSettingsModified(); });
@@ -399,7 +400,7 @@ void UGenerateStaticMeshLODAssetTool::UpdateCollisionVisualization()
 void UGenerateStaticMeshLODAssetTool::CreateNewAsset()
 {
 	check(PreviewWithBackgroundCompute->HaveValidResult());
-	GenerateProcess->CalculateDerivedPathName(BasicProperties->GeneratedSuffix);
+	GenerateProcess->CalculateDerivedPathName(BasicProperties->NewAssetName, BasicProperties->GeneratedSuffix);
 
 	check(GenerateProcess->GraphEvalCriticalSection.TryLock());		// No ops should be running
 	GenerateProcess->WriteDerivedAssetData();
@@ -411,13 +412,13 @@ void UGenerateStaticMeshLODAssetTool::CreateNewAsset()
 void UGenerateStaticMeshLODAssetTool::UpdateExistingAsset()
 {
 	check(PreviewWithBackgroundCompute->HaveValidResult());
-	GenerateProcess->CalculateDerivedPathName(BasicProperties->GeneratedSuffix);
+	GenerateProcess->CalculateDerivedPathName(BasicProperties->NewAssetName, BasicProperties->GeneratedSuffix);
 
 	check(GenerateProcess->GraphEvalCriticalSection.TryLock());		// No ops should be running
 
 	// only updated HD source if we have no HD source asset. Otherwise we are overwriting with existing lowpoly LOD0.
 	bool bUpdateHDSource =
-		BasicProperties->bSaveAsHDSource &&
+		BasicProperties->bSaveInputAsHiResSource &&
 		(GenerateProcess->GetSourceStaticMesh()->IsHiResMeshDescriptionValid() == false);
 
 	GenerateProcess->UpdateSourceAsset(bUpdateHDSource);
diff --git a/Engine/Plugins/Experimental/MeshLODToolset/Source/MeshLODToolset/Public/Graphs/GenerateStaticMeshLODProcess.h b/Engine/Plugins/Experimental/MeshLODToolset/Source/MeshLODToolset/Public/Graphs/GenerateStaticMeshLODProcess.h
--- a/Engine/Plugins/Experimental/MeshLODToolset/Source/MeshLODToolset/Public/Graphs/GenerateStaticMeshLODProcess.h
+++ b/Engine/Plugins/Experimental/MeshLODToolset/Source/MeshLODToolset/Public/Graphs/GenerateStaticMeshLODProcess.h
@@ -68,52 +68,56 @@ struct FGenerateStaticMeshLODProcessSettings
 
 	// Filter settings
 
-	UPROPERTY(EditAnywhere, Category = DetailFilter, meta = (DisplayName = "Detail Filtering"))
+	/** Group layer to use for filtering out detail before processing */
+	UPROPERTY(EditAnywhere, Category = DetailFilter, meta = (DisplayName = "Detail filter group layer"))
 	FName FilterGroupLayer = FName(TEXT("PreFilterGroups"));
 
 	// Thicken settings
 
+	/** Weight map used during mesh thickening */
 	UPROPERTY(EditAnywhere, Category = DetailFilter, meta = (DisplayName = "Thicken weight map name"))
 	FName ThickenWeightMapName = FName(TEXT("ThickenWeightMap"));
 
+	/** Amount to thicken the mesh prior to Solidifying. The thicken weight map values are multiplied by this value. */
 	UPROPERTY(EditAnywhere, Category = DetailFilter, meta = (UIMin = "0", UIMax = "100", ClampMin = "0", ClampMax = "1000"))
 	float ThickenAmount = 0.0f;
 
 	// Solidify settings
 
+	/** Target number of voxels along the maximum dimension for Solidify operation */
 	UPROPERTY(EditAnywhere, Category = Solidify, meta = (UIMin = "8", UIMax = "1024", ClampMin = "8", ClampMax = "1024", DisplayName="Voxel Resolution"))
 	int SolidifyVoxelResolution = 64;
 
+	/** Winding number threshold to determine what is considered inside the mesh during Solidify */
 	UPROPERTY(EditAnywhere, Category = Solidify, AdvancedDisplay, meta = (UIMin = "0.1", UIMax = ".9", ClampMin = "-10", ClampMax = "10"))
 	float WindingThreshold = 0.5f;
 
 
 	// Morphology settings
 
-	//UPROPERTY(EditAnywhere, Category = Morphology, meta = (DisplayName = "Voxel Resolution"))
-	//int MorphologyVoxelResolution = 64;
-
+	/** Offset distance in the Morpohological Closure operation */
 	UPROPERTY(EditAnywhere, Category = Morphology, meta = (UIMin = "0", UIMax = "100", ClampMin = "0", ClampMax = "1000", DisplayName = "Closure Distance"))
 	float ClosureDistance = 1.0f;
 
 
 	// Simplify settings
-
+	/** Target triangle count after simplification */
 	UPROPERTY(EditAnywhere, Category = Simplify, meta = (UIMin = "1", ClampMin = "1", DisplayName = "Simplify Tri Count"))
 	int SimplifyTriangleCount = 500;
 
 
 	// UV settings
-
+	/** Maximum number of charts to create in AutoUV */
 	UPROPERTY(EditAnywhere, Category = AutoUV, meta = (DisplayName = "AutoUV Charts", UIMin = 0, UIMax = 1000))
 	int NumAutoUVCharts = 0;
 
 
 	// Bake Settings
-
+	/** Resolution for normal map and texture baking */
 	UPROPERTY(EditAnywhere, Category = Baking , meta = (DisplayName = "Bake Image Res"))
 	EGenerateStaticMeshLODBakeResolution BakeResolution = EGenerateStaticMeshLODBakeResolution::Resolution512;
 
+	/** How far away from the output mesh to search for input mesh during baking */
 	UPROPERTY(EditAnywhere, Category = Baking, meta = (UIMin = "0", UIMax = "100", ClampMin = "0", ClampMax = "1000", DisplayName = "Bake Thickness"))
 	float BakeThickness = 5.0f;
 
@@ -124,33 +128,36 @@ struct FGenerateStaticMeshLODProcessSettings
 
 
 	// Simple collision generator settings
-
+	/** Type of simple collision objects to produce */
 	UPROPERTY(EditAnywhere, Category = Collision, meta = (DisplayName = "Collision Type"))
 	EGenerateStaticMeshLODSimpleCollisionGeometryType CollisionType = EGenerateStaticMeshLODSimpleCollisionGeometryType::ConvexHulls;
 
 	// Convex Hull Settings
-
+	/** Target triangle count for each convex hull after simplification */
 	UPROPERTY(EditAnywhere, Category = Collision, meta = (NoSpinbox = "true", DisplayName = "Convex Tri Count",
 														  EditCondition = "CollisionType == EGenerateStaticMeshLODSimpleCollisionGeometryType::ConvexHulls"))
 	int ConvexTriangleCount = 50;
 
+	/** Whether to subsample input vertices using a regular grid before computing the convex hull */
 	UPROPERTY(EditAnywhere, Category = Collision, meta = (EditCondition = "CollisionType == EGenerateStaticMeshLODSimpleCollisionGeometryType::ConvexHulls"))
 	bool bPrefilterVertices = true;
 
-	/** Grid resolution (along the maximum-length axis) */
+	/** Grid resolution (along the maximum-length axis) for subsampling before computing the convex hull */
 	UPROPERTY(EditAnywhere, Category = Collision, meta = (NoSpinbox = "true", EditCondition = "CollisionType == EGenerateStaticMeshLODSimpleCollisionGeometryType::ConvexHulls && bPrefilterVertices", UIMin = 4, UIMax = 100))
 	int PrefilterGridResolution = 10;
 
 
 	// Swept Convex Hull Settings
-
+	/** Whether to simplify polygons used for swept convex hulls */
 	UPROPERTY(EditAnywhere, Category = Collision, meta = (EditCondition = "CollisionType == EGenerateStaticMeshLODSimpleCollisionGeometryType::SweptHulls"))
 	bool bSimplifyPolygons = true;
 
+	/** Target minumum edge length for simplified swept convex hulls */
 	UPROPERTY(EditAnywhere, Category = Collision, meta = (NoSpinbox = "true", UIMin = "0", UIMax = "10", ClampMin = "0", ClampMax = "100000",
 														  EditCondition = "CollisionType == EGenerateStaticMeshLODSimpleCollisionGeometryType::SweptHulls"))
 	float HullTolerance = 0.1;
 
+	/** Which axis to sweep along when computing swept convex hulls */
 	UPROPERTY(EditAnywhere, Category = Collision, meta = (EditCondition = "CollisionType == EGenerateStaticMeshLODSimpleCollisionGeometryType::SweptHulls"))
 	EGenerateStaticMeshLODProjectedHullAxisMode SweepAxis = EGenerateStaticMeshLODProjectedHullAxisMode::SmallestVolume;
 
@@ -175,9 +182,11 @@ public:
 	const FString& GetSourceAssetPath() const { return SourceAssetPath; }
 	const FString& GetSourceAssetFolder() const { return SourceAssetFolder; }
 	const FString& GetSourceAssetName() const { return SourceAssetName; }
+
 	static FString GetDefaultDerivedAssetSuffix() { return TEXT("_AutoLOD"); }
+	const FString& GetDerivedAssetName() const { return DerivedAssetName; }
 
-	void CalculateDerivedPathName(FString NewAssetSuffix);
+	void CalculateDerivedPathName(const FString& NewAssetBaseName, const FString& NewAssetSuffix);
 
 	bool ComputeDerivedSourceData(FProgressCancel* Progress);
 	const FDynamicMesh3& GetDerivedLOD0Mesh() const { return DerivedLODMesh; }
@@ -251,8 +260,11 @@ protected:
 	TArray<FSourceMaterialInfo> SourceMaterials;
 
 	FString DerivedSuffix;
-	FString DerivedAssetFolder;
 	FString DerivedAssetPath;
+	FString DerivedAssetFolder;
+	FString DerivedAssetName;
+	FString DerivedAssetNameNoSuffix;
+	
 
 	FDynamicMesh3 DerivedLODMesh;				// the new generated LOD0 mesh
 	UE::Geometry::FMeshTangentsd DerivedLODMeshTangents;		// Tangents for DerivedLODMesh
diff --git a/Engine/Plugins/Experimental/MeshLODToolset/Source/MeshLODToolset/Public/Tools/GenerateStaticMeshLODAssetTool.h b/Engine/Plugins/Experimental/MeshLODToolset/Source/MeshLODToolset/Public/Tools/GenerateStaticMeshLODAssetTool.h
--- a/Engine/Plugins/Experimental/MeshLODToolset/Source/MeshLODToolset/Public/Tools/GenerateStaticMeshLODAssetTool.h
+++ b/Engine/Plugins/Experimental/MeshLODToolset/Source/MeshLODToolset/Public/Tools/GenerateStaticMeshLODAssetTool.h
@@ -62,19 +62,22 @@ class MESHLODTOOLSET_API UGenerateStaticMeshLODAssetToolProperties : public UInt
 	GENERATED_BODY()
 
 public:
-	UPROPERTY(EditAnywhere, Category = AssetOptions, meta = (TransientToolProperty))
-	EGenerateLODAssetOutputMode OutputMode = EGenerateLODAssetOutputMode::UpdateExistingAsset;
 
+	/** Whether to modify the static mesh in place or create a new one. */
 	UPROPERTY(EditAnywhere, Category = AssetOptions, meta = (TransientToolProperty))
-	FString OutputName;
+	EGenerateLODAssetOutputMode OutputMode = EGenerateLODAssetOutputMode::UpdateExistingAsset;
 
 	/** Base name for newly-generated asset */
+	UPROPERTY(EditAnywhere, Category = AssetOptions, meta = (TransientToolProperty, EditCondition = "OutputMode == EGenerateLODAssetOutputMode::CreateNewAsset"))
+	FString NewAssetName;
+
+	/** Suffix to append to newly-generated assets */
 	UPROPERTY(EditAnywhere, Category = AssetOptions, meta = (TransientToolProperty))
 	FString GeneratedSuffix;
 
-	/** If true, the high-resolution input mesh is stored as HD source mesh in the Asset */
+	/** If this is checked and the asset doesn't already have a HiRes source, the input mesh will be stored as the HiRes source */
 	UPROPERTY(EditAnywhere, Category = AssetOptions, meta = (EditCondition = "OutputMode == EGenerateLODAssetOutputMode::UpdateExistingAsset"))
-	bool bSaveAsHDSource = true;
+	bool bSaveInputAsHiResSource = true;
 
 	UPROPERTY(EditAnywhere, Category = Settings)
 	FGenerateStaticMeshLODProcessSettings GeneratorSettings;
@@ -82,6 +85,7 @@ public:
 	// ------------
 	// Code copied from UPolygroupLayersProperties
 
+	/** Group layer to use for partitioning the mesh for simple collision generation */
 	UPROPERTY(EditAnywhere, Category = Settings, meta = (GetOptions = GetGroupLayersFunc))
 	FName CollisionGroupLayerName = TEXT("Default");
 
