- AutoLOD: pop up notifications when input assets might not work as expected
- Add our own log categories for geometry flow and LOD toolset modules to avoid spamming LogTemp

#rnx

diff --git a/Engine/Plugins/Experimental/GeometryFlow/Source/GeometryFlowMeshProcessing/Private/GeometryFlowMeshProcessingModule.cpp b/Engine/Plugins/Experimental/GeometryFlow/Source/GeometryFlowMeshProcessing/Private/GeometryFlowMeshProcessingModule.cpp
--- a/Engine/Plugins/Experimental/GeometryFlow/Source/GeometryFlowMeshProcessing/Private/GeometryFlowMeshProcessingModule.cpp
+++ b/Engine/Plugins/Experimental/GeometryFlow/Source/GeometryFlowMeshProcessing/Private/GeometryFlowMeshProcessingModule.cpp
@@ -2,6 +2,8 @@
 
 #include "GeometryFlowMeshProcessingModule.h"
 
+DEFINE_LOG_CATEGORY(LogGeometryFlowMeshProcessing);
+
 #define LOCTEXT_NAMESPACE "FGeometryFlowMeshProcessingModule"
 
 void FGeometryFlowMeshProcessingModule::StartupModule()
diff --git a/Engine/Plugins/Experimental/GeometryFlow/Source/GeometryFlowMeshProcessing/Public/GeometryFlowMeshProcessingModule.h b/Engine/Plugins/Experimental/GeometryFlow/Source/GeometryFlowMeshProcessing/Public/GeometryFlowMeshProcessingModule.h
--- a/Engine/Plugins/Experimental/GeometryFlow/Source/GeometryFlowMeshProcessing/Public/GeometryFlowMeshProcessingModule.h
+++ b/Engine/Plugins/Experimental/GeometryFlow/Source/GeometryFlowMeshProcessing/Public/GeometryFlowMeshProcessingModule.h
@@ -5,6 +5,8 @@
 #include "CoreMinimal.h"
 #include "Modules/ModuleManager.h"
 
+GEOMETRYFLOWMESHPROCESSING_API DECLARE_LOG_CATEGORY_EXTERN(LogGeometryFlowMeshProcessing, Log, All);
+
 class FGeometryFlowMeshProcessingModule : public IModuleInterface
 {
 public:
diff --git a/Engine/Plugins/Experimental/GeometryFlow/Source/GeometryFlowMeshProcessing/Public/MeshProcessingNodes/MeshProcessingBaseNodes.h b/Engine/Plugins/Experimental/GeometryFlow/Source/GeometryFlowMeshProcessing/Public/MeshProcessingNodes/MeshProcessingBaseNodes.h
--- a/Engine/Plugins/Experimental/GeometryFlow/Source/GeometryFlowMeshProcessing/Public/MeshProcessingNodes/MeshProcessingBaseNodes.h
+++ b/Engine/Plugins/Experimental/GeometryFlow/Source/GeometryFlowMeshProcessing/Public/MeshProcessingNodes/MeshProcessingBaseNodes.h
@@ -2,6 +2,7 @@
 
 #pragma once
 
+#include "GeometryFlowMeshProcessingModule.h"
 #include "GeometryFlowCoreNodes.h"
 #include "GeometryFlowNodeUtil.h"
 #include "DataTypes/DynamicMeshData.h"
@@ -77,7 +78,7 @@ public:
 					bool bIsMeshMutable = DatasIn.GetDataFlags(InParamMesh()).bIsMutableData;
 					if (bIsMeshMutable)
 					{
-						UE_LOG(LogTemp, Warning, TEXT("[%s]  RECOMPUTING MeshOp In Place!"), *GetIdentifier());
+						UE_LOG(LogGeometryFlowMeshProcessing, Display, TEXT("[%s]  RECOMPUTING MeshOp In Place!"), *GetIdentifier());
 
 						FDynamicMesh3 EditableMesh;
 						MeshArg->GiveTo<FDynamicMesh3>(EditableMesh, (int)EMeshProcessingDataTypes::DynamicMesh);
@@ -90,7 +91,7 @@ public:
 					}
 					else
 					{
-						UE_LOG(LogTemp, Warning, TEXT("[%s]  RECOMPUTING MeshOp"), *GetIdentifier());
+						UE_LOG(LogGeometryFlowMeshProcessing, Display, TEXT("[%s]  RECOMPUTING MeshOp"), *GetIdentifier());
 
 						// do we ever want to support using a copy of the source mesh?
 						const FDynamicMesh3& SourceMesh = MeshArg->GetDataConstRef<FDynamicMesh3>((int)EMeshProcessingDataTypes::DynamicMesh);
@@ -191,7 +192,7 @@ public:
 					bool bIsMeshMutable = DatasIn.GetDataFlags(InParamMesh()).bIsMutableData;
 					if (bIsMeshMutable)
 					{
-						UE_LOG(LogTemp, Warning, TEXT("[%s]  RECOMPUTING MeshOp In Place!"), *GetIdentifier());
+						UE_LOG(LogGeometryFlowMeshProcessing, Display, TEXT("[%s]  RECOMPUTING MeshOp In Place!"), *GetIdentifier());
 
 						FDynamicMesh3 EditableMesh;
 						MeshArg->GiveTo<FDynamicMesh3>(EditableMesh, (int)EMeshProcessingDataTypes::DynamicMesh);
@@ -204,7 +205,7 @@ public:
 					}
 					else
 					{
-						UE_LOG(LogTemp, Warning, TEXT("[%s]  RECOMPUTING MeshOp"), *GetIdentifier());
+						UE_LOG(LogGeometryFlowMeshProcessing, Display, TEXT("[%s]  RECOMPUTING MeshOp"), *GetIdentifier());
 
 						// do we ever want to support using a copy of the source mesh?
 						const FDynamicMesh3& SourceMesh = MeshArg->GetDataConstRef<FDynamicMesh3>((int)EMeshProcessingDataTypes::DynamicMesh);
diff --git a/Engine/Plugins/Experimental/MeshLODToolset/Source/MeshLODToolset/MeshLODToolset.Build.cs b/Engine/Plugins/Experimental/MeshLODToolset/Source/MeshLODToolset/MeshLODToolset.Build.cs
--- a/Engine/Plugins/Experimental/MeshLODToolset/Source/MeshLODToolset/MeshLODToolset.Build.cs
+++ b/Engine/Plugins/Experimental/MeshLODToolset/Source/MeshLODToolset/MeshLODToolset.Build.cs
@@ -71,7 +71,8 @@ public class MeshLODToolset : ModuleRules
 
 				"UnrealEd",
 				"EditorScriptingUtilities",
-				"AssetTools"
+				"AssetTools",
+				"Slate"			// for FSlateNotificationManager, etc.
 
 				// ... add private dependencies that you statically link with here ...
 			}
diff --git a/Engine/Plugins/Experimental/MeshLODToolset/Source/MeshLODToolset/Private/Graphs/GenerateMeshLODGraph.cpp b/Engine/Plugins/Experimental/MeshLODToolset/Source/MeshLODToolset/Private/Graphs/GenerateMeshLODGraph.cpp
--- a/Engine/Plugins/Experimental/MeshLODToolset/Source/MeshLODToolset/Private/Graphs/GenerateMeshLODGraph.cpp
+++ b/Engine/Plugins/Experimental/MeshLODToolset/Source/MeshLODToolset/Private/Graphs/GenerateMeshLODGraph.cpp
@@ -2,7 +2,7 @@
 
 #include "Graphs/GenerateMeshLODGraph.h"
 
-
+#include "MeshLODToolsetModule.h"
 #include "GeometryFlowGraph.h"
 #include "GeometryFlowGraphUtil.h"
 #include "BaseNodes/TransferNode.h"
@@ -135,7 +135,7 @@ void FGenerateMeshLODGraph::EvaluateResult(
 	}
 
 	ensure(NormalMapEvalResult == EGeometryFlowResult::Ok);
-	UE_LOG(LogTemp, Warning, TEXT("NormalMapPass - Evaluated %d Nodes, Recomputed %d"), NormalMapEvalInfo->NumEvaluations(), NormalMapEvalInfo->NumComputes());
+	UE_LOG(LogMeshLODToolset, Display, TEXT("NormalMapPass - Evaluated %d Nodes, Recomputed %d"), NormalMapEvalInfo->NumEvaluations(), NormalMapEvalInfo->NumComputes());
 
 
 	//
@@ -158,7 +158,7 @@ void FGenerateMeshLODGraph::EvaluateResult(
 		TextureImages.Add(MoveTemp(NewImage));
 		ensure(TexBakeEvalResult == EGeometryFlowResult::Ok);
 
-		UE_LOG(LogTemp, Warning, TEXT("TextureBakePass %s - Evaluated %d Nodes, Recomputed %d"), *TexBakeStep.Identifier, TexBakeEvalInfo->NumEvaluations(), TexBakeEvalInfo->NumComputes());
+		UE_LOG(LogMeshLODToolset, Display, TEXT("TextureBakePass %s - Evaluated %d Nodes, Recomputed %d"), *TexBakeStep.Identifier, TexBakeEvalInfo->NumEvaluations(), TexBakeEvalInfo->NumComputes());
 	}
 
 	// 
@@ -179,7 +179,7 @@ void FGenerateMeshLODGraph::EvaluateResult(
 	}
 
 	ensure(TangentsEvalResult == EGeometryFlowResult::Ok);
-	UE_LOG(LogTemp, Warning, TEXT("OutputTangentsPass - Evaluated %d Nodes, Recomputed %d"), TangentsEvalInfo->NumEvaluations(), TangentsEvalInfo->NumComputes());
+	UE_LOG(LogMeshLODToolset, Display, TEXT("OutputTangentsPass - Evaluated %d Nodes, Recomputed %d"), TangentsEvalInfo->NumEvaluations(), TangentsEvalInfo->NumComputes());
 
 
 	//
@@ -201,7 +201,7 @@ void FGenerateMeshLODGraph::EvaluateResult(
 	}
 
 	ensure(EvalResult == EGeometryFlowResult::Ok);
-	UE_LOG(LogTemp, Warning, TEXT("OutputMeshPass - Evaluated %d Nodes, Recomputed %d"), MeshEvalInfo->NumEvaluations(), MeshEvalInfo->NumComputes());
+	UE_LOG(LogMeshLODToolset, Display, TEXT("OutputMeshPass - Evaluated %d Nodes, Recomputed %d"), MeshEvalInfo->NumEvaluations(), MeshEvalInfo->NumComputes());
 
 	//
 	// evaluate collision
@@ -221,7 +221,7 @@ void FGenerateMeshLODGraph::EvaluateResult(
 	}
 
 	ensure(CollisionEvalResult == EGeometryFlowResult::Ok);
-	UE_LOG(LogTemp, Warning, TEXT("OutputCollisionPass - Evaluated %d Nodes, Recomputed %d"), CollisionEvalInfo->NumEvaluations(), CollisionEvalInfo->NumComputes());
+	UE_LOG(LogMeshLODToolset, Display, TEXT("OutputCollisionPass - Evaluated %d Nodes, Recomputed %d"), CollisionEvalInfo->NumEvaluations(), CollisionEvalInfo->NumComputes());
 
 }
 
@@ -444,7 +444,7 @@ void FGenerateMeshLODGraph::BuildGraph()
 	//{
 	//	return !Node->GetIdentifier().EndsWith("Settings");
 	//});
-	//UE_LOG(LogTemp, Warning, TEXT("GRAPH:\n%s"), *GraphDump);
+	//UE_LOG(LogMeshLODToolset, Display, TEXT("GRAPH:\n%s"), *GraphDump);
 
 }
 
diff --git a/Engine/Plugins/Experimental/MeshLODToolset/Source/MeshLODToolset/Private/Graphs/GenerateStaticMeshLODProcess.cpp b/Engine/Plugins/Experimental/MeshLODToolset/Source/MeshLODToolset/Private/Graphs/GenerateStaticMeshLODProcess.cpp
--- a/Engine/Plugins/Experimental/MeshLODToolset/Source/MeshLODToolset/Private/Graphs/GenerateStaticMeshLODProcess.cpp
+++ b/Engine/Plugins/Experimental/MeshLODToolset/Source/MeshLODToolset/Private/Graphs/GenerateStaticMeshLODProcess.cpp
@@ -2,6 +2,8 @@
 
 #include "Graphs/GenerateStaticMeshLODProcess.h"
 
+#include "MeshLODToolsetModule.h"
+
 #include "Async/Async.h"
 #include "Async/ParallelFor.h"
 
@@ -178,6 +180,7 @@ bool UGenerateStaticMeshLODProcess::Initialize(UStaticMesh* StaticMeshIn, FProgr
 			{
 				FText WarningText = FText::Format(LOCTEXT("UnusedMaterialWarning", "Found an unused material ({0}). Consider removing it before using this tool."),
 												  FText::FromName(Materials[UnusedMaterialIndex].MaterialInterface->GetFName()));
+				UE_LOG(LogMeshLODToolset, Warning, TEXT("%s"), *WarningText.ToString());
 				Progress->AddWarning(WarningText, FProgressCancel::EMessageLevel::UserWarning);
 			}
 		}
@@ -234,13 +237,11 @@ bool UGenerateStaticMeshLODProcess::Initialize(UStaticMesh* StaticMeshIn, FProgr
 			}
 			if (bFoundTextureNonParamExpession)
 			{
-				FString WarningMessage("WARNING: UGenerateStaticMeshLODProcess: Non-parameter texture sampler detected in input material [%s]. Output materials may have unexpected behaviour.");
-				UE_LOG(LogTemp, Warning, TEXT("%s"), *WarningMessage, *Material->GetName());
-
+				FText WarningText = FText::Format(LOCTEXT("NonParameterTextureWarning", "Non-parameter texture sampler detected in input material [{0}]. Output materials may have unexpected behaviour."),
+												  FText::FromString(Material->GetName()));
+				UE_LOG(LogMeshLODToolset, Warning, TEXT("%s"), *WarningText.ToString());
 				if (Progress)
 				{
-					FText WarningText = FText::Format(LOCTEXT("NonParameterTextureWarning", "Non-parameter texture sampler detected in input material [{0}]. Output materials may have unexpected behaviour."),
-								  FText::FromString(Material->GetName()));
 					Progress->AddWarning(WarningText, FProgressCancel::EMessageLevel::UserWarning);
 				}
 			}
@@ -270,12 +271,8 @@ bool UGenerateStaticMeshLODProcess::Initialize(UStaticMesh* StaticMeshIn, FProgr
 					TexInfo.bIsNormalMap = TexInfo.Texture->IsNormalMap();
 					SourceMaterials[mi].bHasNormalMap |= TexInfo.bIsNormalMap;
 
-					//TexInfo.bIsDefaultTexture = TexInfo.Texture->IsDefaultTexture();
 					TexInfo.bIsDefaultTexture = UEditorAssetLibrary::GetPathNameForLoadedAsset(TexInfo.Texture).StartsWith(TEXT("/Engine/"));
 
-					UE_LOG(LogTemp, Warning, TEXT("SOURCE TEX PATH IS %s - DefaultTex %d - SRGB %d"), *UEditorAssetLibrary::GetPathNameForLoadedAsset(TexInfo.Texture),
-						TexInfo.bIsDefaultTexture?1:0, TexInfo.Texture->SRGB);
-
 					TexInfo.bShouldBakeTexture = (TexInfo.bIsNormalMap == false && TexInfo.bIsDefaultTexture == false);
 					if (TexInfo.bShouldBakeTexture)
 					{
diff --git a/Engine/Plugins/Experimental/MeshLODToolset/Source/MeshLODToolset/Private/MeshLODToolsetModule.cpp b/Engine/Plugins/Experimental/MeshLODToolset/Source/MeshLODToolset/Private/MeshLODToolsetModule.cpp
--- a/Engine/Plugins/Experimental/MeshLODToolset/Source/MeshLODToolset/Private/MeshLODToolsetModule.cpp
+++ b/Engine/Plugins/Experimental/MeshLODToolset/Source/MeshLODToolset/Private/MeshLODToolsetModule.cpp
@@ -4,6 +4,8 @@
 
 #define LOCTEXT_NAMESPACE "FMeshLODToolsetModule"
 
+DEFINE_LOG_CATEGORY(LogMeshLODToolset);
+
 void FMeshLODToolsetModule::StartupModule()
 {
 	// This code will execute after your module is loaded into memory; the exact timing is specified in the .uplugin file per-module
diff --git a/Engine/Plugins/Experimental/MeshLODToolset/Source/MeshLODToolset/Private/Tools/GenerateStaticMeshLODAssetTool.cpp b/Engine/Plugins/Experimental/MeshLODToolset/Source/MeshLODToolset/Private/Tools/GenerateStaticMeshLODAssetTool.cpp
--- a/Engine/Plugins/Experimental/MeshLODToolset/Source/MeshLODToolset/Private/Tools/GenerateStaticMeshLODAssetTool.cpp
+++ b/Engine/Plugins/Experimental/MeshLODToolset/Source/MeshLODToolset/Private/Tools/GenerateStaticMeshLODAssetTool.cpp
@@ -43,6 +43,9 @@
 #include "Polygroups/PolygroupSet.h"
 #include "Polygroups/PolygroupUtil.h"
 
+#include "Framework/Notifications/NotificationManager.h"
+#include "Framework/Docking/TabManager.h"
+#include "Widgets/Notifications/SNotificationList.h"
 #include "TargetInterfaces/PrimitiveComponentBackedTarget.h"
 #include "TargetInterfaces/StaticMeshBackedTarget.h"
 #include "ToolTargetManager.h"
@@ -324,6 +327,16 @@ void UGenerateStaticMeshLODAssetTool::Setup()
 
 	CollisionPreview = NewObject<UPreviewGeometry>(this);
 	CollisionPreview->CreateInWorld(TargetWorld, PreviewTransform);
+
+	// Pop up notifications for any warnings
+	for ( const FProgressCancel::FMessageInfo& Warning : Progress.Warnings )
+	{
+		FNotificationInfo NotificationInfo(Warning.MessageText);
+		NotificationInfo.ExpireDuration = 6.0f;
+		NotificationInfo.Hyperlink = FSimpleDelegate::CreateLambda([]() { FGlobalTabmanager::Get()->TryInvokeTab(FName("OutputLog")); });
+		NotificationInfo.HyperlinkText = LOCTEXT("ShowOutputLogHyperlink", "Show Output Log");
+		FSlateNotificationManager::Get().AddNotification(NotificationInfo)->SetCompletionState(SNotificationItem::CS_Fail);
+	}
 }
 
 void UGenerateStaticMeshLODAssetTool::OnSettingsModified()
diff --git a/Engine/Plugins/Experimental/MeshLODToolset/Source/MeshLODToolset/Public/MeshLODToolsetModule.h b/Engine/Plugins/Experimental/MeshLODToolset/Source/MeshLODToolset/Public/MeshLODToolsetModule.h
--- a/Engine/Plugins/Experimental/MeshLODToolset/Source/MeshLODToolset/Public/MeshLODToolsetModule.h
+++ b/Engine/Plugins/Experimental/MeshLODToolset/Source/MeshLODToolset/Public/MeshLODToolsetModule.h
@@ -5,6 +5,8 @@
 #include "CoreMinimal.h"
 #include "Modules/ModuleManager.h"
 
+DECLARE_LOG_CATEGORY_EXTERN(LogMeshLODToolset, Log, All);
+
 class FMeshLODToolsetModule : public IModuleInterface
 {
 public:
