Minor code and UI changes for AutoLOD.
- DebugDumpGraph was being called before texture baking nodes were in the graph
- Fix some capitalization of tool properties
- Added some code comments
- If morphology distance is zero, don't apply it

#rnx

diff --git a/Engine/Plugins/Experimental/GeometryFlow/Source/GeometryFlowMeshProcessing/Public/MeshDecompositionNodes/MakeTriangleSetsNode.h b/Engine/Plugins/Experimental/GeometryFlow/Source/GeometryFlowMeshProcessing/Public/MeshDecompositionNodes/MakeTriangleSetsNode.h
--- a/Engine/Plugins/Experimental/GeometryFlow/Source/GeometryFlowMeshProcessing/Public/MeshDecompositionNodes/MakeTriangleSetsNode.h
+++ b/Engine/Plugins/Experimental/GeometryFlow/Source/GeometryFlowMeshProcessing/Public/MeshDecompositionNodes/MakeTriangleSetsNode.h
@@ -79,7 +79,9 @@ public:
 
 
 
-
+///
+///  Gather triangles into sets based on their GroupIDs. Optionally exclude any triangles belonging to specified groups.
+///
 class FMakeTriangleSetsFromGroupsNode : public FMakeTriangleSetsFromMeshNode
 {
 public:
diff --git a/Engine/Plugins/Experimental/GeometryFlow/Source/GeometryFlowMeshProcessing/Public/MeshProcessingNodes/MeshVoxMorphologyNode.h b/Engine/Plugins/Experimental/GeometryFlow/Source/GeometryFlowMeshProcessing/Public/MeshProcessingNodes/MeshVoxMorphologyNode.h
--- a/Engine/Plugins/Experimental/GeometryFlow/Source/GeometryFlowMeshProcessing/Public/MeshProcessingNodes/MeshVoxMorphologyNode.h
+++ b/Engine/Plugins/Experimental/GeometryFlow/Source/GeometryFlowMeshProcessing/Public/MeshProcessingNodes/MeshVoxMorphologyNode.h
@@ -48,6 +48,12 @@ public:
 
 	void ApplyMorphology(const FDynamicMesh3& MeshIn, const FVoxMorphologyOpSettings& Settings, FDynamicMesh3& MeshOut)
 	{
+		if (Settings.Distance == 0.0f)
+		{
+			MeshOut = MeshIn;
+			return;
+		}
+
 		FAxisAlignedBox3d Bounds = MeshIn.GetBounds();
 		FDynamicMeshAABBTree3 MeshBVTree(&MeshIn);
 
diff --git a/Engine/Plugins/Experimental/MeshLODToolset/Source/MeshLODToolset/Private/Graphs/GenerateMeshLODGraph.cpp b/Engine/Plugins/Experimental/MeshLODToolset/Source/MeshLODToolset/Private/Graphs/GenerateMeshLODGraph.cpp
--- a/Engine/Plugins/Experimental/MeshLODToolset/Source/MeshLODToolset/Private/Graphs/GenerateMeshLODGraph.cpp
+++ b/Engine/Plugins/Experimental/MeshLODToolset/Source/MeshLODToolset/Private/Graphs/GenerateMeshLODGraph.cpp
@@ -438,14 +438,6 @@ void FGenerateMeshLODGraph::BuildGraph()
 
 	TArray<float> Weights;
 	UpdateThickenWeightMap(Weights);
-
-
-	//FString GraphDump = Graph->DebugDumpGraph([](TSafeSharedPtr<FNode> Node)
-	//{
-	//	return !Node->GetIdentifier().EndsWith("Settings");
-	//});
-	//UE_LOG(LogMeshLODToolset, Display, TEXT("GRAPH:\n%s"), *GraphDump);
-
 }
 
 
@@ -460,7 +452,7 @@ int32 FGenerateMeshLODGraph::AppendTextureBakeNode(const TImageBuilder<FVector4f
 	// add source node
 	NewNode.TexSourceNode = Graph->AddNodeOfType<FTextureImageSourceNode>(FString::Printf(TEXT("TextureSource%d_%s"), NewNode.Index, *NewNode.Identifier));
 
-	// normal map baker
+	// texture baker
 	NewNode.BakeNode = Graph->AddNodeOfType<FBakeMeshTextureImageNode>(FString::Printf(TEXT("BakeTexImage%d_%s"), NewNode.Index, *NewNode.Identifier));
 	ensure(Graph->InferConnection(BakeCacheNode, NewNode.BakeNode) == EGeometryFlowResult::Ok);
 	ensure(Graph->InferConnection(NewNode.TexSourceNode, NewNode.BakeNode) == EGeometryFlowResult::Ok);
diff --git a/Engine/Plugins/Experimental/MeshLODToolset/Source/MeshLODToolset/Private/Graphs/GenerateStaticMeshLODProcess.cpp b/Engine/Plugins/Experimental/MeshLODToolset/Source/MeshLODToolset/Private/Graphs/GenerateStaticMeshLODProcess.cpp
--- a/Engine/Plugins/Experimental/MeshLODToolset/Source/MeshLODToolset/Private/Graphs/GenerateStaticMeshLODProcess.cpp
+++ b/Engine/Plugins/Experimental/MeshLODToolset/Source/MeshLODToolset/Private/Graphs/GenerateStaticMeshLODProcess.cpp
@@ -39,11 +39,13 @@
 #include "Components/StaticMeshComponent.h"
 #include "PhysicsEngine/BodySetup.h"
 #include "StaticMeshAttributes.h"
+#include "GeometryFlowTypes.h"
 
 #define LOCTEXT_NAMESPACE "UGenerateStaticMeshLODProcess"
 
 #include "ExplicitUseGeometryMathTypes.h"		// using UE::Geometry::(math types)
 using namespace UE::Geometry;
+using namespace UE::GeometryFlow;
 
 namespace
 {
@@ -428,7 +430,6 @@ bool UGenerateStaticMeshLODProcess::InitializeGenerator()
 			}
 		}
 	}
-	
 
 	// initialize source mesh
 	Generator->SetSourceMesh(this->SourceMesh);
diff --git a/Engine/Plugins/Experimental/MeshLODToolset/Source/MeshLODToolset/Public/Graphs/GenerateMeshLODGraph.h b/Engine/Plugins/Experimental/MeshLODToolset/Source/MeshLODToolset/Public/Graphs/GenerateMeshLODGraph.h
--- a/Engine/Plugins/Experimental/MeshLODToolset/Source/MeshLODToolset/Public/Graphs/GenerateMeshLODGraph.h
+++ b/Engine/Plugins/Experimental/MeshLODToolset/Source/MeshLODToolset/Public/Graphs/GenerateMeshLODGraph.h
@@ -73,6 +73,8 @@ public:
 
 protected:
 
+	friend class UGenerateStaticMeshLODProcess;
+
 	TUniquePtr<UE::GeometryFlow::FGraph> Graph;
 
 
diff --git a/Engine/Plugins/Experimental/MeshLODToolset/Source/MeshLODToolset/Public/Graphs/GenerateStaticMeshLODProcess.h b/Engine/Plugins/Experimental/MeshLODToolset/Source/MeshLODToolset/Public/Graphs/GenerateStaticMeshLODProcess.h
--- a/Engine/Plugins/Experimental/MeshLODToolset/Source/MeshLODToolset/Public/Graphs/GenerateStaticMeshLODProcess.h
+++ b/Engine/Plugins/Experimental/MeshLODToolset/Source/MeshLODToolset/Public/Graphs/GenerateStaticMeshLODProcess.h
@@ -69,13 +69,13 @@ struct FGenerateStaticMeshLODProcessSettings
 	// Filter settings
 
 	/** Group layer to use for filtering out detail before processing */
-	UPROPERTY(EditAnywhere, Category = DetailFilter, meta = (DisplayName = "Detail filter group layer"))
+	UPROPERTY(EditAnywhere, Category = DetailFilter, meta = (DisplayName = "Detail Filter Group Layer"))
 	FName FilterGroupLayer = FName(TEXT("PreFilterGroups"));
 
 	// Thicken settings
 
 	/** Weight map used during mesh thickening */
-	UPROPERTY(EditAnywhere, Category = DetailFilter, meta = (DisplayName = "Thicken weight map name"))
+	UPROPERTY(EditAnywhere, Category = DetailFilter, meta = (DisplayName = "Thicken Weight Map"))
 	FName ThickenWeightMapName = FName(TEXT("ThickenWeightMap"));
 
 	/** Amount to thicken the mesh prior to Solidifying. The thicken weight map values are multiplied by this value. */
@@ -108,7 +108,7 @@ struct FGenerateStaticMeshLODProcessSettings
 
 	// UV settings
 	/** Maximum number of charts to create in AutoUV */
-	UPROPERTY(EditAnywhere, Category = AutoUV, meta = (DisplayName = "AutoUV Charts", UIMin = 0, UIMax = 1000))
+	UPROPERTY(EditAnywhere, Category = AutoUV, meta = (DisplayName = "Max AutoUV Charts", UIMin = 0, UIMax = 1000))
 	int NumAutoUVCharts = 0;
 
 
