GeometryProcessing:
* Fix crash when checking for unused materials during AutoLOD
* Fix error printed about missing texture just before texture is generated during AutoLOD

#rnx

diff --git a/Engine/Plugins/Experimental/MeshLODToolset/Source/MeshLODToolset/Private/Graphs/GenerateStaticMeshLODProcess.cpp b/Engine/Plugins/Experimental/MeshLODToolset/Source/MeshLODToolset/Private/Graphs/GenerateStaticMeshLODProcess.cpp
--- a/Engine/Plugins/Experimental/MeshLODToolset/Source/MeshLODToolset/Private/Graphs/GenerateStaticMeshLODProcess.cpp
+++ b/Engine/Plugins/Experimental/MeshLODToolset/Source/MeshLODToolset/Private/Graphs/GenerateStaticMeshLODProcess.cpp
@@ -98,14 +98,25 @@ namespace GenerateStaticMeshLODProcessHelpers
 		}
 	}
 
-	void FindUnreferencedMaterials(UStaticMesh* StaticMesh, FMeshDescription* MeshDescription, TArray<int32>& OutUnreferencedMaterials)
+	TArray<int32> FindUnreferencedMaterials(UStaticMesh* StaticMesh, const FMeshDescription* MeshDescription)
 	{
-		const TArray<FStaticMaterial>& CurMaterialSet = StaticMesh->GetStaticMaterials();
-		const int32 NumMaterials = CurMaterialSet.Num();
-		TArray<bool> MatUsedFlags;
-		MatUsedFlags.Init(false, NumMaterials);
+		const TArray<FStaticMaterial>& MaterialSet = StaticMesh->GetStaticMaterials();
+		const int32 NumMaterials = MaterialSet.Num();
+
+		auto IsValidMaterial = [&MaterialSet](int32 MaterialID) {
+			return MaterialSet[MaterialID].MaterialInterface != nullptr;
+		};
+
+		// Initially flag only valid materials as potentially unused.
+		TArray<bool> MatUnusedFlags;
+		MatUnusedFlags.SetNum(NumMaterials);
+		int32 NumMatUnused = 0;
+		for (int32 MaterialID = 0; MaterialID < NumMaterials; ++MaterialID)
+		{
+			MatUnusedFlags[MaterialID] = IsValidMaterial(MaterialID);
+			NumMatUnused += MatUnusedFlags[MaterialID];
+		}
 
-		TArray<int32> MeshMaterialIndices;
 		TMap<FPolygonGroupID, int32> PolygonGroupToMaterialIndex;
 		const FStaticMeshConstAttributes Attributes(*MeshDescription);
 		TPolygonGroupAttributesConstRef<FName> PolygonGroupImportedMaterialSlotNames = Attributes.GetPolygonGroupMaterialSlotNames();
@@ -124,17 +135,25 @@ namespace GenerateStaticMeshLODProcessHelpers
 		{
 			const FPolygonGroupID PolygonGroupID = MeshDescription->GetTrianglePolygonGroup(TriangleID);
 			const int32 MaterialIndex = PolygonGroupToMaterialIndex[PolygonGroupID];
-			MatUsedFlags[MaterialIndex] = true;
-			MeshMaterialIndices.AddUnique(MaterialIndex);
+			bool& bMatUnusedFlag = MatUnusedFlags[MaterialIndex];
+			NumMatUnused -= static_cast<int32>(bMatUnusedFlag);
+			bMatUnusedFlag = false;
+			if (NumMatUnused == 0)
+			{
+				break;
+			}
 		}
 
+		TArray<int32> UnreferencedMaterials;
+		UnreferencedMaterials.Reserve(NumMatUnused);
 		for (int32 MaterialID = 0; MaterialID < NumMaterials; ++MaterialID)
 		{
-			if (!MatUsedFlags[MaterialID])
+			if (MatUnusedFlags[MaterialID])
 			{
-				OutUnreferencedMaterials.Emplace(MaterialID);
+				UnreferencedMaterials.Emplace(MaterialID);
 			}
 		}
+		return UnreferencedMaterials;
 	}
 }
 
@@ -174,14 +193,13 @@ bool UGenerateStaticMeshLODProcess::Initialize(UStaticMesh* StaticMeshIn, FProgr
 	// warn the user if there are any unsed materials in the mesh
 	if (Progress)
 	{
-		TArray<int32> UnusedMaterials;
-		GenerateStaticMeshLODProcessHelpers::FindUnreferencedMaterials(SourceStaticMesh, SourceMeshDescription.Get(), UnusedMaterials);
-		if (UnusedMaterials.Num() > 0)
+		for (int32 UnusedMaterialIndex : GenerateStaticMeshLODProcessHelpers::FindUnreferencedMaterials(SourceStaticMesh, SourceMeshDescription.Get()))
 		{
-			for (int32 UnusedMaterialIndex : UnusedMaterials)
+			const TObjectPtr<class UMaterialInterface> MaterialInterface = Materials[UnusedMaterialIndex].MaterialInterface;
+			if (ensure(MaterialInterface != nullptr))
 			{
 				FText WarningText = FText::Format(LOCTEXT("UnusedMaterialWarning", "Found an unused material ({0}). Consider removing it before using this tool."),
-												  FText::FromName(Materials[UnusedMaterialIndex].MaterialInterface->GetFName()));
+					                              FText::FromName(MaterialInterface->GetFName()));
 				UE_LOG(LogMeshLODToolset, Warning, TEXT("%s"), *WarningText.ToString());
 				Progress->AddWarning(WarningText, FProgressCancel::EMessageLevel::UserWarning);
 			}
@@ -791,29 +809,32 @@ void UGenerateStaticMeshLODProcess::UpdateSourceAsset(bool bSetNewHDSourceAsset)
 
 bool UGenerateStaticMeshLODProcess::IsSourceAsset(const FString& AssetPath) const
 {
-	FAssetData AssetData = UEditorAssetLibrary::FindAssetData(AssetPath);
-
-	for (const FSourceMaterialInfo& MaterialInfo : SourceMaterials)
+	if (UEditorAssetLibrary::DoesAssetExist(AssetPath))
 	{
-		UMaterialInterface* MaterialInterface = MaterialInfo.SourceMaterial.MaterialInterface;
-		if (MaterialInterface == nullptr)
-		{
-			continue;
-		}
+		const FAssetData AssetData = UEditorAssetLibrary::FindAssetData(AssetPath);
 
-		FString SourceMaterialPath = UEditorAssetLibrary::GetPathNameForLoadedAsset(MaterialInterface);
-		if (UEditorAssetLibrary::FindAssetData(SourceMaterialPath) == AssetData)
+		for (const FSourceMaterialInfo& MaterialInfo : SourceMaterials)
 		{
-			return true;
-		}
+			UMaterialInterface* MaterialInterface = MaterialInfo.SourceMaterial.MaterialInterface;
+			if (MaterialInterface == nullptr)
+			{
+				continue;
+			}
 
-		for (const FTextureInfo& TextureInfo : MaterialInfo.SourceTextures)
-		{
-			FString SourceTexturePath = UEditorAssetLibrary::GetPathNameForLoadedAsset(TextureInfo.Texture);
-			if (UEditorAssetLibrary::FindAssetData(SourceTexturePath) == AssetData)
+			FString SourceMaterialPath = UEditorAssetLibrary::GetPathNameForLoadedAsset(MaterialInterface);
+			if (UEditorAssetLibrary::FindAssetData(SourceMaterialPath) == AssetData)
 			{
 				return true;
 			}
+
+			for (const FTextureInfo& TextureInfo : MaterialInfo.SourceTextures)
+			{
+				FString SourceTexturePath = UEditorAssetLibrary::GetPathNameForLoadedAsset(TextureInfo.Texture);
+				if (UEditorAssetLibrary::FindAssetData(SourceTexturePath) == AssetData)
+				{
+					return true;
+				}
+			}
 		}
 	}
 
