Add multi-texture baking to AutoLOD

diff --git a/Engine/Plugins/Experimental/GeometryFlow/Source/GeometryFlowCore/Public/GeometryFlowGraph.h b/Engine/Plugins/Experimental/GeometryFlow/Source/GeometryFlowCore/Public/GeometryFlowGraph.h
--- a/Engine/Plugins/Experimental/GeometryFlow/Source/GeometryFlowCore/Public/GeometryFlowGraph.h
+++ b/Engine/Plugins/Experimental/GeometryFlow/Source/GeometryFlowCore/Public/GeometryFlowGraph.h
@@ -41,7 +41,8 @@ public:
 
 	struct FHandle
 	{
-		int32 Identifier;
+		static const int32 InvalidHandle = -1;
+		int32 Identifier = InvalidHandle;
 		bool operator==(const FHandle& OtherHandle) const { return Identifier == OtherHandle.Identifier; }
 	};
 
diff --git a/Engine/Plugins/Experimental/GeometryFlow/Source/GeometryFlowMeshProcessing/Private/MeshBakingNodes/BakeMeshMultiTextureNode.cpp b/Engine/Plugins/Experimental/GeometryFlow/Source/GeometryFlowMeshProcessing/Private/MeshBakingNodes/BakeMeshMultiTextureNode.cpp
new file mode 100644
--- /dev/null
+++ b/Engine/Plugins/Experimental/GeometryFlow/Source/GeometryFlowMeshProcessing/Private/MeshBakingNodes/BakeMeshMultiTextureNode.cpp
@@ -0,0 +1,56 @@
+// Copyright Epic Games, Inc. All Rights Reserved.
+
+#include "MeshBakingNodes/BakeMeshMultiTextureNode.h"
+#include "GeometryFlowNodeUtil.h"
+
+#include "Sampling/MeshResampleImageBaker.h"
+
+using namespace UE::Geometry;
+using namespace UE::GeometryFlow;
+
+void FBakeMeshMultiTextureNode::Evaluate(
+	const FNamedDataMap& DatasIn,
+	FNamedDataMap& DatasOut,
+	TUniquePtr<FEvaluationInfo>& EvaluationInfo)
+{
+	if (ensure(DatasOut.Contains(OutParamTextureImage())))
+	{
+		bool bAllInputsValid = true;
+		bool bRecomputeRequired = (IsOutputAvailable(OutParamTextureImage()) == false);
+		TSafeSharedPtr<IData> BakeCacheArg = FindAndUpdateInputForEvaluate(InParamBakeCache(), DatasIn, bRecomputeRequired, bAllInputsValid);
+		TSafeSharedPtr<IData> MaterialToTextureMapArg = FindAndUpdateInputForEvaluate(InParamMaterialTextures(), DatasIn, bRecomputeRequired, bAllInputsValid);
+		TSafeSharedPtr<IData> SettingsArg = FindAndUpdateInputForEvaluate(InParamSettings(), DatasIn, bRecomputeRequired, bAllInputsValid);
+
+		if (bAllInputsValid)
+		{
+			if (bRecomputeRequired)
+			{
+				// always make a copy of settings
+				FBakeMeshMultiTextureSettings Settings;
+				SettingsArg->GetDataCopy(Settings, FBakeMeshMultiTextureSettings::DataTypeIdentifier);
+
+				const FMeshBakingCache& BakeCacheContainer = BakeCacheArg->GetDataConstRef<FMeshBakingCache>((int)EMeshProcessingDataTypes::BakingCache);
+				const FDynamicMesh3& DetailMesh = BakeCacheContainer.DetailMesh;
+				check(DetailMesh.HasAttributes() && DetailMesh.Attributes()->NumUVLayers() > Settings.DetailUVLayer);
+				const FDynamicMeshUVOverlay* DetailUVOverlay = DetailMesh.Attributes()->GetUVLayer(Settings.DetailUVLayer);
+
+				const FMaterialIDToTextureMap& SourceMaterialTextures = 
+					MaterialToTextureMapArg->GetDataConstRef<FMaterialIDToTextureMap>(FMaterialIDToTextureMap::DataTypeIdentifier);
+
+				FMeshMultiResampleImageBaker Baker;
+				Baker.SetCache(&BakeCacheContainer.BakeCache);
+				Baker.DetailUVOverlay = DetailUVOverlay;
+				Baker.MultiTextures = SourceMaterialTextures.MaterialIDTextureMap;
+				Baker.Bake();
+
+				FTextureImage ResultImage;
+				ResultImage.Image = MoveTemp(*Baker.TakeResult());
+
+				SetOutput(OutParamTextureImage(), MakeMovableData<FTextureImage>(MoveTemp(ResultImage)) );
+
+				EvaluationInfo->CountCompute(this);
+			}
+			DatasOut.SetData(OutParamTextureImage(), GetOutput(OutParamTextureImage()));
+		}
+	}
+}
diff --git a/Engine/Plugins/Experimental/GeometryFlow/Source/GeometryFlowMeshProcessing/Public/DataTypes/TextureImageData.h b/Engine/Plugins/Experimental/GeometryFlow/Source/GeometryFlowMeshProcessing/Public/DataTypes/TextureImageData.h
--- a/Engine/Plugins/Experimental/GeometryFlow/Source/GeometryFlowMeshProcessing/Public/DataTypes/TextureImageData.h
+++ b/Engine/Plugins/Experimental/GeometryFlow/Source/GeometryFlowMeshProcessing/Public/DataTypes/TextureImageData.h
@@ -26,6 +26,15 @@ struct FTextureImage
 GEOMETRYFLOW_DECLARE_BASIC_TYPES(TextureImage, FTextureImage, (int)EMeshProcessingDataTypes::TextureImage)
 
 
+struct FMaterialIDToTextureMap
+{
+	DECLARE_GEOMETRYFLOW_DATA_TYPE_IDENTIFIER(EMeshProcessingDataTypes::MaterialIDToTextureMap);
+
+	TMap<int32, TSafeSharedPtr<TImageBuilder<FVector4f>>> MaterialIDTextureMap;
+};
+
+GEOMETRYFLOW_DECLARE_BASIC_TYPES(MaterialIDToTextureMap, FMaterialIDToTextureMap, (int)EMeshProcessingDataTypes::MaterialIDToTextureMap)
+
 
 }	// end namespace GeometryFlow
 }	// end namespace UE
\ No newline at end of file
diff --git a/Engine/Plugins/Experimental/GeometryFlow/Source/GeometryFlowMeshProcessing/Public/MeshBakingNodes/BakeMeshMultiTextureNode.h b/Engine/Plugins/Experimental/GeometryFlow/Source/GeometryFlowMeshProcessing/Public/MeshBakingNodes/BakeMeshMultiTextureNode.h
new file mode 100644
--- /dev/null
+++ b/Engine/Plugins/Experimental/GeometryFlow/Source/GeometryFlowMeshProcessing/Public/MeshBakingNodes/BakeMeshMultiTextureNode.h
@@ -0,0 +1,54 @@
+// Copyright Epic Games, Inc. All Rights Reserved.
+
+#pragma once
+
+#include "MeshProcessingNodes/MeshProcessingBaseNodes.h"
+#include "MeshProcessingNodes/MeshProcessingDataTypes.h"
+#include "MeshProcessingNodes/MeshTangentsNodes.h"
+#include "MeshBakingNodes/BakeMeshTextureImageNode.h"
+#include "DataTypes/MeshImageBakingData.h"
+#include "DataTypes/TextureImageData.h"
+
+class UTexture2D;
+
+namespace UE
+{
+namespace GeometryFlow
+{
+
+struct GEOMETRYFLOWMESHPROCESSING_API FBakeMeshMultiTextureSettings : public FBakeMeshTextureImageSettings
+{
+	DECLARE_GEOMETRYFLOW_DATA_TYPE_IDENTIFIER(EMeshProcessingDataTypes::BakeMultiTextureSettings);
+};
+
+GEOMETRYFLOW_DECLARE_SETTINGS_TYPES(FBakeMeshMultiTextureSettings, BakeMeshMultiTexture);
+
+
+class GEOMETRYFLOWMESHPROCESSING_API FBakeMeshMultiTextureNode : public FNode
+{
+protected:
+	using SettingsDataType = TMovableData<FBakeMeshMultiTextureSettings, FBakeMeshMultiTextureSettings::DataTypeIdentifier>;
+
+public:
+	static const FString InParamBakeCache() { return TEXT("BakeCache"); }
+	static const FString InParamMaterialTextures() { return TEXT("MaterialIDToTextureMap"); }
+	static const FString InParamSettings() { return TEXT("Settings"); }
+	static const FString OutParamTextureImage() { return TEXT("TextureImage"); }
+
+public:
+	FBakeMeshMultiTextureNode()
+	{
+		AddInput(InParamBakeCache(), MakeUnique<TImmutableNodeInput<FMeshBakingCache, FMeshBakingCache::DataTypeIdentifier>>());
+		AddInput(InParamMaterialTextures(), MakeUnique<TBasicNodeInput<FMaterialIDToTextureMap, FMaterialIDToTextureMap::DataTypeIdentifier>>());
+		AddInput(InParamSettings(), MakeUnique<TBasicNodeInput<FBakeMeshMultiTextureSettings, FBakeMeshMultiTextureSettings::DataTypeIdentifier>>());
+		AddOutput(OutParamTextureImage(), MakeUnique<TBasicNodeOutput<FTextureImage, FTextureImage::DataTypeIdentifier>>());
+	}
+
+	virtual void Evaluate(
+		const FNamedDataMap& DatasIn,
+		FNamedDataMap& DatasOut,
+		TUniquePtr<FEvaluationInfo>& EvaluationInfo) override;
+};
+
+}	// end namespace GeometryFlow
+}	// end namespace UE
\ No newline at end of file
diff --git a/Engine/Plugins/Experimental/GeometryFlow/Source/GeometryFlowMeshProcessing/Public/MeshProcessingNodes/MeshProcessingDataTypes.h b/Engine/Plugins/Experimental/GeometryFlow/Source/GeometryFlowMeshProcessing/Public/MeshProcessingNodes/MeshProcessingDataTypes.h
--- a/Engine/Plugins/Experimental/GeometryFlow/Source/GeometryFlowMeshProcessing/Public/MeshProcessingNodes/MeshProcessingDataTypes.h
+++ b/Engine/Plugins/Experimental/GeometryFlow/Source/GeometryFlowMeshProcessing/Public/MeshProcessingNodes/MeshProcessingDataTypes.h
@@ -22,6 +22,7 @@ enum class EMeshProcessingDataTypes
 	CollisionGeometry = (int)EDataTypes::BaseMeshProcessingTypes + 13,
 	IndexSets = (int)EDataTypes::BaseMeshProcessingTypes + 14,
 	WeightMap = (int)EDataTypes::BaseMeshProcessingTypes + 15,
+	MaterialIDToTextureMap = (int)EDataTypes::BaseMeshProcessingTypes + 16,
 
 	SolidifySettings = (int)EDataTypes::BaseMeshProcessingTypes + 100,
 	VoxMorphologyOpSettings = (int)EDataTypes::BaseMeshProcessingTypes + 101,
@@ -43,6 +44,7 @@ enum class EMeshProcessingDataTypes
 	MakeBakingCacheSettings = (int)EDataTypes::BaseMeshProcessingTypes + 160,
 	BakeNormalMapSettings = (int)EDataTypes::BaseMeshProcessingTypes + 161,
 	BakeTextureImageSettings = (int)EDataTypes::BaseMeshProcessingTypes + 162,
+	BakeMultiTextureSettings = (int)EDataTypes::BaseMeshProcessingTypes + 163,
 
 
 	GenerateCollisionConvexHullsSettings = (int)EDataTypes::BaseMeshProcessingTypes + 200,
diff --git a/Engine/Plugins/Experimental/MeshLODToolset/Source/MeshLODToolset/Private/Graphs/GenerateMeshLODGraph.cpp b/Engine/Plugins/Experimental/MeshLODToolset/Source/MeshLODToolset/Private/Graphs/GenerateMeshLODGraph.cpp
--- a/Engine/Plugins/Experimental/MeshLODToolset/Source/MeshLODToolset/Private/Graphs/GenerateMeshLODGraph.cpp
+++ b/Engine/Plugins/Experimental/MeshLODToolset/Source/MeshLODToolset/Private/Graphs/GenerateMeshLODGraph.cpp
@@ -25,6 +25,7 @@
 #include "DataTypes/MeshImageBakingData.h"
 #include "MeshBakingNodes/BakeMeshNormalMapNode.h"
 #include "MeshBakingNodes/BakeMeshTextureImageNode.h"
+#include "MeshBakingNodes/BakeMeshMultiTextureNode.h"
 
 #include "MeshDecompositionNodes/MakeTriangleSetsNode.h"
 #include "PhysicsNodes/GenerateSimpleCollisionNode.h"
@@ -110,6 +111,7 @@ void FGenerateMeshLODGraph::EvaluateResult(
 	FSimpleShapeSet3d& ResultCollision,
 	UE::GeometryFlow::FNormalMapImage& NormalMap,
 	TArray<TUniquePtr<UE::GeometryFlow::FTextureImage>>& TextureImages,
+	UE::GeometryFlow::FTextureImage& MultiTextureImage,
 	FProgressCancel* Progress)
 {
 	//FScopedDurationTimeLogger Timer(TEXT("FGenerateMeshLODGraph::EvaluateResult -- serial execution"));
@@ -161,6 +163,25 @@ void FGenerateMeshLODGraph::EvaluateResult(
 		UE_LOG(LogMeshLODToolset, Display, TEXT("TextureBakePass %s - Evaluated %d Nodes, Recomputed %d"), *TexBakeStep.Identifier, TexBakeEvalInfo->NumEvaluations(), TexBakeEvalInfo->NumComputes());
 	}
 
+
+	// 
+	// evaluate multi texture bake
+	//
+	{
+		TUniquePtr<FEvaluationInfo> MultiTextureBakeEvalInfo = MakeUnique<FEvaluationInfo>();
+		MultiTextureBakeEvalInfo->Progress = Progress;
+		EGeometryFlowResult MultiTextureBakeEvalResult = Graph->EvaluateResult(BakeMultiTextureNode, FBakeMeshMultiTextureNode::OutParamTextureImage(),
+																			   MultiTextureImage, (int)EMeshProcessingDataTypes::TextureImage, MultiTextureBakeEvalInfo, true);
+
+		if (Progress && Progress->Cancelled())
+		{
+			return;
+		}
+
+		ensure(MultiTextureBakeEvalResult == EGeometryFlowResult::Ok);
+		UE_LOG(LogMeshLODToolset, Display, TEXT("MultiTextureBake - Evaluated %d Nodes, Recomputed %d"), MultiTextureBakeEvalInfo->NumEvaluations(), MultiTextureBakeEvalInfo->NumComputes());
+	}
+
 	// 
 	// evaluate tangents
 	//
@@ -467,3 +488,22 @@ int32 FGenerateMeshLODGraph::AppendTextureBakeNode(const TImageBuilder<FVector4f
 
 	return NewNode.Index;
 }
+
+
+void FGenerateMeshLODGraph::AppendMultiTextureBakeNode(const TMap<int32, TSafeSharedPtr<TImageBuilder<FVector4f>>>& SourceMaterialImages)
+{
+	// add source node
+	MaterialIDTextureSourceNode = Graph->AddNodeOfType<FMaterialIDToTextureMapSourceNode>(FString::Printf(TEXT("MaterialIDTextureSource")));
+
+	// texture baker
+	BakeMultiTextureNode = Graph->AddNodeOfType<FBakeMeshMultiTextureNode>(FString::Printf(TEXT("BakeMultiTexture")));
+	ensure(Graph->InferConnection(BakeCacheNode, BakeMultiTextureNode) == EGeometryFlowResult::Ok);
+	ensure(Graph->InferConnection(MaterialIDTextureSourceNode, BakeMultiTextureNode) == EGeometryFlowResult::Ok);
+	BakeMultiTextureSettingsNode = Graph->AddNodeOfType<FBakeMeshMultiTextureSettingsSourceNode>(TEXT("BakeMultiTextureSettings"));
+	ensure(Graph->InferConnection(BakeMultiTextureSettingsNode, BakeMultiTextureNode) == EGeometryFlowResult::Ok);
+
+	FMaterialIDToTextureMap InputMap;
+	InputMap.MaterialIDTextureMap = SourceMaterialImages;
+
+	UpdateSourceNodeValue<FMaterialIDToTextureMapSourceNode>(*Graph, MaterialIDTextureSourceNode, InputMap);
+}
diff --git a/Engine/Plugins/Experimental/MeshLODToolset/Source/MeshLODToolset/Private/Graphs/GenerateStaticMeshLODProcess.cpp b/Engine/Plugins/Experimental/MeshLODToolset/Source/MeshLODToolset/Private/Graphs/GenerateStaticMeshLODProcess.cpp
--- a/Engine/Plugins/Experimental/MeshLODToolset/Source/MeshLODToolset/Private/Graphs/GenerateStaticMeshLODProcess.cpp
+++ b/Engine/Plugins/Experimental/MeshLODToolset/Source/MeshLODToolset/Private/Graphs/GenerateStaticMeshLODProcess.cpp
@@ -35,6 +35,8 @@
 #include "MaterialGraph/MaterialGraphNode.h"
 #include "WeightMapUtil.h"
 
+#include "Sampling/MeshResampleImageBaker.h"
+
 #include "Engine/StaticMesh.h"
 #include "Components/StaticMeshComponent.h"
 #include "PhysicsEngine/BodySetup.h"
@@ -426,6 +428,69 @@ void UGenerateStaticMeshLODProcess::CalculateDerivedPathName(const FString& NewA
 }
 
 
+int UGenerateStaticMeshLODProcess::SelectTextureToBake(const TArray<FTextureInfo>& TextureInfos) const
+{
+	TArray<int> TextureVotes;
+	TextureVotes.Init(0, TextureInfos.Num());
+
+	for (int TextureIndex = 0; TextureIndex < TextureInfos.Num(); ++TextureIndex)
+	{
+		FTextureInfo Info = TextureInfos[TextureIndex];
+
+		if (!Info.bIsNormalMap)
+		{
+			++TextureVotes[TextureIndex];
+		}
+
+		if (Info.ParameterName == "Diffuse" || Info.ParameterName == "Albedo" || Info.ParameterName == "BaseColor")
+		{
+			++TextureVotes[TextureIndex];
+		}
+
+		if (Info.Dimensions.GetWidth() > 0 && Info.Dimensions.GetHeight() > 0)
+		{
+			++TextureVotes[TextureIndex];
+		}
+
+		UTexture2D* Tex2D = Info.Texture;
+
+		if (Tex2D)
+		{
+			// Texture uses SRGB
+			if (Tex2D->SRGB != 0)
+			{
+				++TextureVotes[TextureIndex];
+			}
+
+#if WITH_EDITORONLY_DATA
+			// Texture has multiple channels
+			ETextureSourceFormat Format = Tex2D->Source.GetFormat();
+			if (Format == TSF_BGRA8 || Format == TSF_BGRE8 || Format == TSF_RGBA16 || Format == TSF_RGBA16F)
+			{
+				++TextureVotes[TextureIndex];
+			}
+#endif
+
+			// What else? Largest texture? Most layers? Most mipmaps?
+		}
+	}
+
+	int MaxIndex = -1;
+	int MaxVotes = -1;
+	for (int TextureIndex = 0; TextureIndex < TextureInfos.Num(); ++TextureIndex)
+	{
+		if (TextureVotes[TextureIndex] > MaxVotes)
+		{
+			MaxIndex = TextureIndex;
+			MaxVotes = TextureVotes[TextureIndex];
+		}
+	}
+
+	return MaxIndex;
+}
+
+
+
 bool UGenerateStaticMeshLODProcess::InitializeGenerator()
 {
 	Generator = MakeUnique<FGenerateMeshLODGraph>();
@@ -449,6 +514,37 @@ bool UGenerateStaticMeshLODProcess::InitializeGenerator()
 		}
 	}
 
+	//  Add multi-texture bake node	
+	TMap<int32, TSafeSharedPtr<UE::Geometry::TImageBuilder<FVector4f>>> SourceMaterialImages;
+
+	// TODO: Relying on MaterialID being the index in the SourceMaterials array is brittle
+	for (int32 MaterialID = 0; MaterialID < SourceMaterials.Num(); ++MaterialID)
+	{
+		FSourceMaterialInfo& MatInfo = SourceMaterials[MaterialID];
+
+		if (MatInfo.bIsPreviouslyGeneratedMaterial == false && MatInfo.bIsReusable == false && MatInfo.bHasTexturesToBake)
+		{
+			int TexIndex = SelectTextureToBake(MatInfo.SourceTextures);
+			if (TexIndex < 0)
+			{
+				continue;
+			}
+
+			FTextureInfo& TextureInfo = MatInfo.SourceTextures[TexIndex];
+
+			MultiTextureParameterName.Add(MaterialID, TextureInfo.ParameterName);
+
+			TSafeSharedPtr<TImageBuilder<FVector4f>> ImagePtr =
+				MakeSafeShared<TImageBuilder<FVector4f>>(TextureInfo.Image);
+			SourceMaterialImages.Add(MaterialID, ImagePtr);
+
+			TextureInfo.bIsUsedInMultiTextureBaking = true;
+		}
+	}
+
+	Generator->AppendMultiTextureBakeNode(SourceMaterialImages);
+
+
 	// initialize source mesh
 	Generator->SetSourceMesh(this->SourceMesh);
 
@@ -557,7 +653,8 @@ void UGenerateStaticMeshLODProcess::UpdateSettings(const FGenerateStaticMeshLODP
 
 
 	if ( (NewCombinedSettings.BakeResolution != CurrentSettings.BakeResolution) ||
-		 (NewCombinedSettings.BakeThickness != CurrentSettings.BakeThickness))
+		 (NewCombinedSettings.BakeThickness != CurrentSettings.BakeThickness) || 
+		 (NewCombinedSettings.bCombineTextures != CurrentSettings.bCombineTextures))
 	{
 		UE::GeometryFlow::FMeshMakeBakingCacheSettings NewBakeSettings = Generator->GetCurrentBakeCacheSettings();
 		NewBakeSettings.Dimensions = FImageDimensions((int32)NewCombinedSettings.BakeResolution, (int32)NewCombinedSettings.BakeResolution);
@@ -605,6 +702,7 @@ bool UGenerateStaticMeshLODProcess::ComputeDerivedSourceData(FProgressCancel* Pr
 		this->DerivedCollision,
 		this->DerivedNormalMapImage,
 		this->DerivedTextureImages,
+		this->DerivedMultiTextureBakeImage,
 		Progress);
 
 	if (Progress && Progress->Cancelled())
@@ -681,6 +779,12 @@ void UGenerateStaticMeshLODProcess::GetDerivedMaterialsPreview(FPreviewMaterials
 		for (int32 ti = 0; ti < NumTextures; ++ti)
 		{
 			const FTextureInfo& SourceTex = SourceMaterialInfo.SourceTextures[ti];
+
+			if (CurrentSettings.bCombineTextures && SourceTex.bIsUsedInMultiTextureBaking)
+			{
+				continue;
+			}
+
 			bool bConvertToSRGB = SourceTex.Texture->SRGB;
 			const FTextureInfo& DerivedTex = DerivedMaterialInfo.DerivedTextures[ti];
 			if (DerivedTex.bShouldBakeTexture)
@@ -708,6 +812,22 @@ void UGenerateStaticMeshLODProcess::GetDerivedMaterialsPreview(FPreviewMaterials
 	UTexture2D* PreviewNormalMapTex = NormapMapBuilder.GetTexture2D();
 	MaterialSetOut.Textures.Add(PreviewNormalMapTex);
 
+	// create multi-texture bake result
+	if (CurrentSettings.bCombineTextures)
+	{
+		bool bConvertToSRGB = true;
+
+		FTexture2DBuilder MultiTextureBuilder;
+		MultiTextureBuilder.Initialize(FTexture2DBuilder::ETextureType::Color, DerivedMultiTextureBakeImage.Image.GetDimensions());
+		MultiTextureBuilder.GetTexture2D()->SRGB = bConvertToSRGB;
+		MultiTextureBuilder.Copy(DerivedMultiTextureBakeImage.Image, bConvertToSRGB);
+		MultiTextureBuilder.Commit(false);
+
+		DerivedMultiTextureBakeResult = MultiTextureBuilder.GetTexture2D();
+		MaterialSetOut.Textures.Add(DerivedMultiTextureBakeResult);
+	}
+
+
 	// create derived MIDs and point to new textures
 	for (int32 mi = 0; mi < NumMaterials; ++mi)
 	{
@@ -726,8 +846,15 @@ void UGenerateStaticMeshLODProcess::GetDerivedMaterialsPreview(FPreviewMaterials
 			// rewrite texture parameters to new textures
 			UpdateMaterialTextureParameters(GeneratedMID, SourceMaterialInfo, SourceToPreviewTexMap, PreviewNormalMapTex);
 
+			if (CurrentSettings.bCombineTextures && DerivedMultiTextureBakeResult && MultiTextureParameterName.Contains(mi))
+			{
+				FMaterialParameterInfo ParamInfo(MultiTextureParameterName[mi]);
+				GeneratedMID->SetTextureParameterValueByInfo(ParamInfo, DerivedMultiTextureBakeResult);
+			}
+
 			MaterialSetOut.Materials.Add(GeneratedMID);
 		}
+
 	}
 
 }
@@ -755,6 +882,11 @@ void UGenerateStaticMeshLODProcess::UpdateMaterialTextureParameters(
 		}
 		else if (SourceTex.bShouldBakeTexture)
 		{
+			if (CurrentSettings.bCombineTextures && SourceTex.bIsUsedInMultiTextureBaking)
+			{
+				continue;
+			}
+
 			UTexture2D*const* FoundTexture = PreviewTextures.Find(SourceTex.Texture);
 			if (ensure(FoundTexture))
 			{
@@ -928,6 +1060,28 @@ void UGenerateStaticMeshLODProcess::WriteDerivedTextures(bool bCreatingNewStatic
 		}
 	}
 
+
+	// write multi-texture bake result
+	if (CurrentSettings.bCombineTextures)
+	{
+		bool bConvertToSRGB = true;
+		
+		FTexture2DBuilder MultiTextureBuilder;
+		MultiTextureBuilder.Initialize(FTexture2DBuilder::ETextureType::Color, DerivedMultiTextureBakeImage.Image.GetDimensions());
+		MultiTextureBuilder.GetTexture2D()->SRGB = bConvertToSRGB;
+		MultiTextureBuilder.Copy(DerivedMultiTextureBakeImage.Image, bConvertToSRGB);
+		MultiTextureBuilder.Commit(false);
+
+		DerivedMultiTextureBakeResult = MultiTextureBuilder.GetTexture2D();
+		if (ensure(DerivedMultiTextureBakeResult))
+		{
+			FTexture2DBuilder::CopyPlatformDataToSourceData(DerivedMultiTextureBakeResult, FTexture2DBuilder::ETextureType::Color);
+
+			// write asset
+			bool bWriteOK = WriteDerivedTexture(DerivedMultiTextureBakeResult, DerivedAssetNameNoSuffix + TEXT("_MultiTexture"), bCreatingNewStaticMeshAsset);
+			ensure(bWriteOK);
+		}
+	}
 	
 }
 
@@ -1081,6 +1235,13 @@ void UGenerateStaticMeshLODProcess::WriteDerivedMaterials(bool bCreatingNewStati
 		DerivedMaterialInfo.DerivedMaterial.MaterialInterface = GeneratedMIC;
 		DerivedMaterialInfo.DerivedMaterial.MaterialSlotName = FName(FString::Printf(TEXT("GeneratedMat%d"), mi));
 		DerivedMaterialInfo.DerivedMaterial.ImportedMaterialSlotName = DerivedMaterialInfo.DerivedMaterial.MaterialSlotName;
+
+		if (CurrentSettings.bCombineTextures && DerivedMultiTextureBakeResult && MultiTextureParameterName.Contains(mi))
+		{
+			FMaterialParameterInfo ParamInfo(MultiTextureParameterName[mi]);
+			GeneratedMIC->SetTextureParameterValueEditorOnly(ParamInfo, DerivedMultiTextureBakeResult);
+		}
+
 	}
 }
 
diff --git a/Engine/Plugins/Experimental/MeshLODToolset/Source/MeshLODToolset/Private/Tools/GenerateStaticMeshLODAssetTool.cpp b/Engine/Plugins/Experimental/MeshLODToolset/Source/MeshLODToolset/Private/Tools/GenerateStaticMeshLODAssetTool.cpp
--- a/Engine/Plugins/Experimental/MeshLODToolset/Source/MeshLODToolset/Private/Tools/GenerateStaticMeshLODAssetTool.cpp
+++ b/Engine/Plugins/Experimental/MeshLODToolset/Source/MeshLODToolset/Private/Tools/GenerateStaticMeshLODAssetTool.cpp
@@ -254,6 +254,7 @@ void UGenerateStaticMeshLODAssetTool::Setup()
 	BasicProperties->WatchProperty(BasicProperties->GeneratorSettings.NumAutoUVCharts, [this](int) { OnSettingsModified(); });
 	BasicProperties->WatchProperty(BasicProperties->GeneratorSettings.BakeResolution, [this](EGenerateStaticMeshLODBakeResolution) { OnSettingsModified(); });
 	BasicProperties->WatchProperty(BasicProperties->GeneratorSettings.BakeThickness, [this](float) { OnSettingsModified(); });
+	BasicProperties->WatchProperty(BasicProperties->GeneratorSettings.bCombineTextures, [this](bool) { OnSettingsModified(); });
 	BasicProperties->WatchProperty(BasicProperties->GeneratorSettings.CollisionType, [this](EGenerateStaticMeshLODSimpleCollisionGeometryType) { OnSettingsModified(); });
 	BasicProperties->WatchProperty(BasicProperties->GeneratorSettings.ConvexTriangleCount, [this](int) { OnSettingsModified(); });
 	BasicProperties->WatchProperty(BasicProperties->GeneratorSettings.bPrefilterVertices, [this](bool) { OnSettingsModified(); });
diff --git a/Engine/Plugins/Experimental/MeshLODToolset/Source/MeshLODToolset/Public/Graphs/GenerateMeshLODGraph.h b/Engine/Plugins/Experimental/MeshLODToolset/Source/MeshLODToolset/Public/Graphs/GenerateMeshLODGraph.h
--- a/Engine/Plugins/Experimental/MeshLODToolset/Source/MeshLODToolset/Public/Graphs/GenerateMeshLODGraph.h
+++ b/Engine/Plugins/Experimental/MeshLODToolset/Source/MeshLODToolset/Public/Graphs/GenerateMeshLODGraph.h
@@ -3,6 +3,7 @@
 #pragma once
 
 #include "GeometryFlowGraph.h"
+#include "GeometryFlowTypes.h"
 #include "DataTypes/NormalMapData.h"
 #include "DataTypes/TextureImageData.h"
 #include "DataTypes/CollisionGeometryData.h"
@@ -33,6 +34,8 @@ public:
 
 	int32 AppendTextureBakeNode(const UE::Geometry::TImageBuilder<FVector4f>& SourceImage, const FString& Identifier);
 
+	void AppendMultiTextureBakeNode(const TMap<int32, UE::GeometryFlow::TSafeSharedPtr<UE::Geometry::TImageBuilder<FVector4f>>>& SourceMaterialImages);
+
 	void SetSourceMesh(const FDynamicMesh3& SourceMesh);
 
 	void UpdatePreFilterSettings(const FMeshLODGraphPreFilterSettings& PreFilterSettings);
@@ -69,6 +72,7 @@ public:
 		UE::Geometry::FSimpleShapeSet3d& ResultCollision,
 		UE::GeometryFlow::FNormalMapImage& NormalMap,
 		TArray<TUniquePtr<UE::GeometryFlow::FTextureImage>>& TextureImages,
+		UE::GeometryFlow::FTextureImage& MultiTextureImage,
 		FProgressCancel* Progress);
 
 protected:
@@ -119,6 +123,10 @@ protected:
 	UE::GeometryFlow::FGraph::FHandle BakeNormalMapNode;
 	UE::GeometryFlow::FGraph::FHandle BakeNormalMapSettingsNode;
 
+	UE::GeometryFlow::FGraph::FHandle BakeMultiTextureNode;
+	UE::GeometryFlow::FGraph::FHandle BakeMultiTextureSettingsNode;
+	UE::GeometryFlow::FGraph::FHandle MaterialIDTextureSourceNode;
+
 	UE::GeometryFlow::FGraph::FHandle ThickenNode;
 	UE::GeometryFlow::FGraph::FHandle ThickenSettingsNode;
 	UE::GeometryFlow::FGraph::FHandle ThickenWeightMapNode;
diff --git a/Engine/Plugins/Experimental/MeshLODToolset/Source/MeshLODToolset/Public/Graphs/GenerateStaticMeshLODProcess.h b/Engine/Plugins/Experimental/MeshLODToolset/Source/MeshLODToolset/Public/Graphs/GenerateStaticMeshLODProcess.h
--- a/Engine/Plugins/Experimental/MeshLODToolset/Source/MeshLODToolset/Public/Graphs/GenerateStaticMeshLODProcess.h
+++ b/Engine/Plugins/Experimental/MeshLODToolset/Source/MeshLODToolset/Public/Graphs/GenerateStaticMeshLODProcess.h
@@ -121,13 +121,17 @@ struct FGenerateStaticMeshLODProcessSettings
 	UPROPERTY(EditAnywhere, Category = Baking, meta = (UIMin = "0", UIMax = "100", ClampMin = "0", ClampMax = "1000", DisplayName = "Bake Thickness"))
 	float BakeThickness = 5.0f;
 
+	UPROPERTY(EditAnywhere, Category = Baking, meta = (DisplayName = "Combine Textures"))
+	bool bCombineTextures = true;
+
+
+	// Simple collision generator settings
+
 	// Transient property, not set directly by the user. The user controls a CollisionGroupLayerName dropdown property
 	// on the Tool and that value is copied here.
 	UPROPERTY(meta = (TransientToolProperty))
 	FName CollisionGroupLayerName = TEXT("Default");
 
-
-	// Simple collision generator settings
 	/** Type of simple collision objects to produce */
 	UPROPERTY(EditAnywhere, Category = Collision, meta = (DisplayName = "Collision Type"))
 	EGenerateStaticMeshLODSimpleCollisionGeometryType CollisionType = EGenerateStaticMeshLODSimpleCollisionGeometryType::ConvexHulls;
@@ -241,8 +245,11 @@ protected:
 		bool bIsNormalMap = false;
 		bool bIsDefaultTexture = false;
 		bool bShouldBakeTexture = false;
+		bool bIsUsedInMultiTextureBaking = false;
 	};
 
+	int SelectTextureToBake(const TArray<FTextureInfo>& TextureInfos) const;
+
 	// Information about one of the input StaticMesh Materials. Computed in Initialize() and not modified afterwards
 	struct FSourceMaterialInfo
 	{
@@ -273,6 +280,7 @@ protected:
 	// Texture set potentially required by output Material set
 	UE::GeometryFlow::FNormalMapImage DerivedNormalMapImage;	// Normal Map
 	TArray<TUniquePtr<UE::GeometryFlow::FTextureImage>> DerivedTextureImages;	// generated Textures
+	UE::GeometryFlow::FTextureImage DerivedMultiTextureBakeImage;
 	TMap<UTexture2D*, int32> SourceTextureToDerivedTexIndex;	// mapping from input Textures to DerivedTextureImages index
 
 	// Information about an output Material
@@ -297,6 +305,13 @@ protected:
 	UPROPERTY()
 	UTexture2D* DerivedNormalMapTex;
 
+	// For each material participating in multi-texture baking, the parameter name of the texture
+	TMap<int32, FName> MultiTextureParameterName;
+
+	UPROPERTY()
+	UTexture2D* DerivedMultiTextureBakeResult;
+
+
 	TUniquePtr<FGenerateMeshLODGraph> Generator;			// active LODGenerator Graph
 	bool InitializeGenerator();
 
