AutoLOD: set tangents on preview mesh

#rnx

diff --git a/Engine/Plugins/Experimental/MeshLODToolset/Source/MeshLODToolset/Private/Tools/GenerateStaticMeshLODAssetTool.cpp b/Engine/Plugins/Experimental/MeshLODToolset/Source/MeshLODToolset/Private/Tools/GenerateStaticMeshLODAssetTool.cpp
--- a/Engine/Plugins/Experimental/MeshLODToolset/Source/MeshLODToolset/Private/Tools/GenerateStaticMeshLODAssetTool.cpp
+++ b/Engine/Plugins/Experimental/MeshLODToolset/Source/MeshLODToolset/Private/Tools/GenerateStaticMeshLODAssetTool.cpp
@@ -272,6 +272,8 @@ void UGenerateStaticMeshLODAssetTool::Setup()
 	this->OpFactory = MakeUnique<FGenerateStaticMeshLODAssetOperatorFactory>(this, (UE::Geometry::FTransform3d)PreviewTransform);
 	PreviewWithBackgroundCompute = NewObject<UMeshOpPreviewWithBackgroundCompute>(this, "Preview");
 	PreviewWithBackgroundCompute->Setup(this->TargetWorld, this->OpFactory.Get());
+	PreviewWithBackgroundCompute->PreviewMesh->SetTangentsMode(EDynamicMeshComponentTangentsMode::ExternallyProvided);
+
 
 	// For the first computation, display a bounding box with the working material. Otherwise it looks like nothing
 	// is happening. And we don't want to copy over the potentially huge input mesh to be the preview mesh.
@@ -312,6 +314,26 @@ void UGenerateStaticMeshLODAssetTool::Setup()
 		GenerateProcess->GraphEvalCriticalSection.Unlock();
 	});
 
+
+	PreviewWithBackgroundCompute->OnMeshUpdated.AddLambda([this](const UMeshOpPreviewWithBackgroundCompute* PreviewWithBackgroundCompute)
+	{
+		// GenerateProcess might be in use by an Op somewhere else
+		GenerateProcess->GraphEvalCriticalSection.Lock();
+
+		UE::Geometry::FMeshTangentsd Tangents = GenerateProcess->GetDerivedLOD0MeshTangents();
+		PreviewWithBackgroundCompute->PreviewMesh->EditMesh([&Tangents](FDynamicMesh3& Mesh)
+		{
+			if (Mesh.HasAttributes() && Tangents.GetTangents().Num() > 0 && Tangents.GetBitangents().Num() > 0)
+			{
+				Mesh.Attributes()->SetNumNormalLayers(3);
+				ensure(Tangents.CopyToOverlays(Mesh));
+			}
+		});
+
+		GenerateProcess->GraphEvalCriticalSection.Unlock();
+	});
+
+
 	PreviewWithBackgroundCompute->ConfigureMaterials(
 		ToolSetupUtil::GetDefaultSculptMaterial(GetToolManager()),
 		ToolSetupUtil::GetDefaultWorkingMaterial(GetToolManager())
