AutoLOD: assign computed tangents to the preview mesh inside the Op, before returning to the main thread

#rnx

diff --git a/Engine/Plugins/Experimental/MeshLODToolset/Source/MeshLODToolset/Private/Tools/GenerateStaticMeshLODAssetTool.cpp b/Engine/Plugins/Experimental/MeshLODToolset/Source/MeshLODToolset/Private/Tools/GenerateStaticMeshLODAssetTool.cpp
--- a/Engine/Plugins/Experimental/MeshLODToolset/Source/MeshLODToolset/Private/Tools/GenerateStaticMeshLODAssetTool.cpp
+++ b/Engine/Plugins/Experimental/MeshLODToolset/Source/MeshLODToolset/Private/Tools/GenerateStaticMeshLODAssetTool.cpp
@@ -109,6 +109,12 @@ namespace GenerateStaticMeshLODAssetLocals
 				*ResultMesh = GenerateProcess->GetDerivedLOD0Mesh();
 				ResultTangents = GenerateProcess->GetDerivedLOD0MeshTangents();
 				ResultCollision = GenerateProcess->GetDerivedCollision();
+
+				if (ResultMesh->HasAttributes() && ResultTangents.GetTangents().Num() > 0 && ResultTangents.GetBitangents().Num() > 0)
+				{
+					ResultMesh->Attributes()->SetNumNormalLayers(3);
+					ensure(ResultTangents.CopyToOverlays(*ResultMesh));
+				}
 			};
 
 			GenerateProcess->GraphEvalCriticalSection.Lock();
@@ -314,26 +320,6 @@ void UGenerateStaticMeshLODAssetTool::Setup()
 		GenerateProcess->GraphEvalCriticalSection.Unlock();
 	});
 
-
-	PreviewWithBackgroundCompute->OnMeshUpdated.AddLambda([this](const UMeshOpPreviewWithBackgroundCompute* PreviewCompute)
-	{
-		// GenerateProcess might be in use by an Op somewhere else
-		GenerateProcess->GraphEvalCriticalSection.Lock();
-
-		UE::Geometry::FMeshTangentsd Tangents = GenerateProcess->GetDerivedLOD0MeshTangents();
-		PreviewCompute->PreviewMesh->EditMesh([&Tangents](FDynamicMesh3& Mesh)
-		{
-			if (Mesh.HasAttributes() && Tangents.GetTangents().Num() > 0 && Tangents.GetBitangents().Num() > 0)
-			{
-				Mesh.Attributes()->SetNumNormalLayers(3);
-				ensure(Tangents.CopyToOverlays(Mesh));
-			}
-		});
-
-		GenerateProcess->GraphEvalCriticalSection.Unlock();
-	});
-
-
 	PreviewWithBackgroundCompute->ConfigureMaterials(
 		ToolSetupUtil::GetDefaultSculptMaterial(GetToolManager()),
 		ToolSetupUtil::GetDefaultWorkingMaterial(GetToolManager())
