AutoLOD: add support for PatchBuilder UVs and Simplification based on geometric deviation tolerance, split up settings into separate structs so that settings panel is more organized
#rnx

diff --git a/Engine/Plugins/Experimental/GeometryFlow/Source/GeometryFlowMeshProcessing/Public/MeshProcessingNodes/MeshSimplifyNode.h b/Engine/Plugins/Experimental/GeometryFlow/Source/GeometryFlowMeshProcessing/Public/MeshProcessingNodes/MeshSimplifyNode.h
--- a/Engine/Plugins/Experimental/GeometryFlow/Source/GeometryFlowMeshProcessing/Public/MeshProcessingNodes/MeshSimplifyNode.h
+++ b/Engine/Plugins/Experimental/GeometryFlow/Source/GeometryFlowMeshProcessing/Public/MeshProcessingNodes/MeshSimplifyNode.h
@@ -5,8 +5,10 @@
 #include "MeshProcessingNodes/MeshProcessingBaseNodes.h"
 #include "MeshProcessingNodes/MeshProcessingDataTypes.h"
 
+#include "DynamicMesh/DynamicMeshAABBTree3.h"
 #include "MeshSimplification.h"
 #include "MeshConstraintsUtil.h"
+#include "ProjectionTargets.h"
 
 
 namespace UE
@@ -27,7 +29,8 @@ enum class EMeshSimplifyTargetType
 {
 	TriangleCount = 0,
 	VertexCount = 1,
-	TrianglePercentage = 2
+	TrianglePercentage = 2,
+	GeometricDeviation = 3
 };
 
 struct FMeshSimplifySettings
@@ -39,6 +42,7 @@ struct FMeshSimplifySettings
 
 	int TargetCount = 100;
 	float TargetFraction = 0.5;
+	float GeometricTolerance = 0.5;
 
 	bool bDiscardAttributes = false;
 	bool bPreventNormalFlips = true;
@@ -129,6 +133,19 @@ public:
 		{
 			Simplifier.SimplifyToVertexCount( FMath::Max(3, Settings.TargetCount) );
 		}
+		else if (Settings.TargetType == EMeshSimplifyTargetType::GeometricDeviation)
+		{
+			// need to create projection target to measure error
+			FDynamicMesh3 MeshCopy;
+			MeshCopy.Copy(*TargetMesh, false, false, false, false);
+			FDynamicMeshAABBTree3 MeshCopySpatial(&MeshCopy, true);
+			FMeshProjectionTarget ProjTarget(&MeshCopy, &MeshCopySpatial);
+			Simplifier.SetProjectionTarget(&ProjTarget);
+
+			Simplifier.GeometricErrorConstraint = SimplifierType::EGeometricErrorCriteria::PredictedPointToProjectionTarget;
+			Simplifier.GeometricErrorTolerance = Settings.GeometricTolerance;
+			Simplifier.SimplifyToVertexCount(3);
+		}
 		else
 		{
 			check(false);
diff --git a/Engine/Plugins/Experimental/GeometryFlow/Source/GeometryFlowMeshProcessingEditor/Private/MeshProcessingNodes/MeshAutoGenerateUVsNode.cpp b/Engine/Plugins/Experimental/GeometryFlow/Source/GeometryFlowMeshProcessingEditor/Private/MeshProcessingNodes/MeshAutoGenerateUVsNode.cpp
--- a/Engine/Plugins/Experimental/GeometryFlow/Source/GeometryFlowMeshProcessingEditor/Private/MeshProcessingNodes/MeshAutoGenerateUVsNode.cpp
+++ b/Engine/Plugins/Experimental/GeometryFlow/Source/GeometryFlowMeshProcessingEditor/Private/MeshProcessingNodes/MeshAutoGenerateUVsNode.cpp
@@ -14,10 +14,33 @@ void FMeshAutoGenerateUVsNode::GenerateUVs(const FDynamicMesh3& MeshIn, const FM
 	TSharedPtr<FDynamicMesh3, ESPMode::ThreadSafe> InputMesh = MakeShared<FDynamicMesh3, ESPMode::ThreadSafe>(MeshIn);
 
 	FParameterizeMeshOp ParameterizeMeshOp;
-	ParameterizeMeshOp.Stretch = Settings.Stretch;
-	ParameterizeMeshOp.NumCharts = Settings.NumCharts;
 	ParameterizeMeshOp.InputMesh = InputMesh;
-	ParameterizeMeshOp.Method = UE::Geometry::EParamOpBackend::UVAtlas;
+	ParameterizeMeshOp.Stretch = Settings.UVAtlasStretch;
+	ParameterizeMeshOp.NumCharts = Settings.UVAtlasNumCharts;
+	ParameterizeMeshOp.XAtlasMaxIterations = Settings.XAtlasMaxIterations;
+
+	ParameterizeMeshOp.InitialPatchCount = Settings.NumInitialPatches;
+	ParameterizeMeshOp.PatchCurvatureAlignmentWeight = Settings.CurvatureAlignment;
+	ParameterizeMeshOp.PatchMergingMetricThresh = Settings.MergingThreshold;
+	ParameterizeMeshOp.PatchMergingAngleThresh = Settings.MaxAngleDeviationDeg;
+	ParameterizeMeshOp.ExpMapNormalSmoothingSteps = Settings.SmoothingSteps;
+	ParameterizeMeshOp.ExpMapNormalSmoothingAlpha = Settings.SmoothingAlpha;
+
+	ParameterizeMeshOp.bEnablePacking = Settings.bAutoPack;
+	ParameterizeMeshOp.Width = ParameterizeMeshOp.Height = Settings.PackingTargetWidth;
+	
+	if (Settings.Method == EAutoUVMethod::UVAtlas)
+	{
+		ParameterizeMeshOp.Method = UE::Geometry::EParamOpBackend::UVAtlas;
+	}
+	else if (Settings.Method == EAutoUVMethod::XAtlas)
+	{
+		ParameterizeMeshOp.Method = UE::Geometry::EParamOpBackend::XAtlas;
+	}
+	else
+	{
+		ParameterizeMeshOp.Method = UE::Geometry::EParamOpBackend::PatchBuilder;
+	}
 
 	FProgressCancel Progress;
 	ParameterizeMeshOp.CalculateResult(&Progress);
diff --git a/Engine/Plugins/Experimental/GeometryFlow/Source/GeometryFlowMeshProcessingEditor/Public/MeshProcessingNodes/MeshAutoGenerateUVsNode.h b/Engine/Plugins/Experimental/GeometryFlow/Source/GeometryFlowMeshProcessingEditor/Public/MeshProcessingNodes/MeshAutoGenerateUVsNode.h
--- a/Engine/Plugins/Experimental/GeometryFlow/Source/GeometryFlowMeshProcessingEditor/Public/MeshProcessingNodes/MeshAutoGenerateUVsNode.h
+++ b/Engine/Plugins/Experimental/GeometryFlow/Source/GeometryFlowMeshProcessingEditor/Public/MeshProcessingNodes/MeshAutoGenerateUVsNode.h
@@ -11,13 +11,38 @@ namespace GeometryFlow
 {
 
 
+enum class EAutoUVMethod : uint8
+{
+	PatchBuilder = 0,
+	UVAtlas = 1,
+	XAtlas = 2
+};
+
+
+
 struct GEOMETRYFLOWMESHPROCESSINGEDITOR_API FMeshAutoGenerateUVsSettings
 {
 	DECLARE_GEOMETRYFLOW_DATA_TYPE_IDENTIFIER(EMeshProcessingDataTypesEditor::MeshAutoGenerateUVsSettings);
 
-	double Stretch = 0.5;
-	int NumCharts = 0;
+	EAutoUVMethod Method = EAutoUVMethod::PatchBuilder;
+
+	// UVAtlas parameters
+	double UVAtlasStretch = 0.5;
+	int UVAtlasNumCharts = 0;
+
+	// XAtlas parameters
+	int XAtlasMaxIterations = 1;
+
+	// PatchBuilder parameters
+	int NumInitialPatches = 100;
+	double CurvatureAlignment = 1.0;
+	double MergingThreshold = 1.5;
+	double MaxAngleDeviationDeg = 45.0;
+	int SmoothingSteps = 5;
+	double SmoothingAlpha = 0.25;
 
+	bool bAutoPack = false;
+	int PackingTargetWidth = 512;
 };
 GEOMETRYFLOW_DECLARE_SETTINGS_TYPES(FMeshAutoGenerateUVsSettings, MeshAutoGenerateUVs);
 
diff --git a/Engine/Plugins/Experimental/MeshLODToolset/Source/MeshLODToolset/Private/Graphs/GenerateMeshLODGraph.cpp b/Engine/Plugins/Experimental/MeshLODToolset/Source/MeshLODToolset/Private/Graphs/GenerateMeshLODGraph.cpp
--- a/Engine/Plugins/Experimental/MeshLODToolset/Source/MeshLODToolset/Private/Graphs/GenerateMeshLODGraph.cpp
+++ b/Engine/Plugins/Experimental/MeshLODToolset/Source/MeshLODToolset/Private/Graphs/GenerateMeshLODGraph.cpp
@@ -320,13 +320,15 @@ void FGenerateMeshLODGraph::BuildGraph()
 	AutoUVSettingsNode = Graph->AddNodeOfType<FMeshAutoGenerateUVsSettingsSourceNode>(TEXT("AutoUVSettings"));
 	Graph->InferConnection(AutoUVSettingsNode, AutoUVNode);
 
-	RecomputeUVNode = Graph->AddNodeOfType<FMeshRecalculateUVsNode>(TEXT("RecalcUV"));
-	Graph->InferConnection(AutoUVNode, RecomputeUVNode);
-	RecomputeUVSettingsNode = Graph->AddNodeOfType<FMeshRecalculateUVsSettingsSourceNode>(TEXT("RecalcUVSettings"));
-	Graph->InferConnection(RecomputeUVSettingsNode, RecomputeUVNode);
+	// disable RecomputeUVNode for now
+	//RecomputeUVNode = Graph->AddNodeOfType<FMeshRecalculateUVsNode>(TEXT("RecalcUV"));
+	//Graph->InferConnection(AutoUVNode, RecomputeUVNode);
+	//RecomputeUVSettingsNode = Graph->AddNodeOfType<FMeshRecalculateUVsSettingsSourceNode>(TEXT("RecalcUVSettings"));
+	//Graph->InferConnection(RecomputeUVSettingsNode, RecomputeUVNode);
 
 	RepackUVNode = Graph->AddNodeOfType<FMeshRepackUVsNode>(TEXT("RepackUV"));
-	Graph->InferConnection(RecomputeUVNode, RepackUVNode);
+	//Graph->InferConnection(RecomputeUVNode, RepackUVNode);
+	Graph->InferConnection(AutoUVNode, RepackUVNode);
 	RepackUVSettingsNode = Graph->AddNodeOfType<FMeshRepackUVsSettingsSourceNode>(TEXT("RepackUVSettings"));
 	Graph->InferConnection(RepackUVSettingsNode, RepackUVNode);
 
@@ -404,16 +406,22 @@ void FGenerateMeshLODGraph::BuildGraph()
 	UpdatePreFilterSettings(PreFilterSettings);
 
 	FMeshSolidifySettings SolidifySettings;
+	SolidifySettings.VoxelResolution = 128;
 	UpdateSolidifySettings(SolidifySettings);
 
 	FVoxClosureSettings MorphologySettings;
+	MorphologySettings.VoxelResolution = SolidifySettings.VoxelResolution;
 	MorphologySettings.Distance = 5.0;
 	UpdateMorphologySettings(MorphologySettings);
 
+	// TODO: can use simpler settings if there is only one material...
 	FMeshSimplifySettings SimplifySettings;
 	SimplifySettings.bDiscardAttributes = false;
 	SimplifySettings.SimplifyType = EMeshSimplifyType::AttributeAware;
-	SimplifySettings.TargetType = EMeshSimplifyTargetType::TriangleCount;
+	//SimplifySettings.TargetType = EMeshSimplifyTargetType::TriangleCount;
+	//SimplifySettings.TargetCount = 500;
+	SimplifySettings.TargetType = EMeshSimplifyTargetType::GeometricDeviation;
+	SimplifySettings.GeometricTolerance = 0.5;
 	SimplifySettings.TargetCount = 500;
 	SimplifySettings.MaterialBorderConstraints = EEdgeRefineFlags::NoFlip;
 	UpdateSimplifySettings(SimplifySettings);
@@ -424,12 +432,22 @@ void FGenerateMeshLODGraph::BuildGraph()
 	UpdateSettingsSourceNodeValue(*Graph, NormalsSettingsNode, NormalsSettings);
 
 	FMeshAutoGenerateUVsSettings AutoUVSettings;
-	AutoUVSettings.NumCharts = 20;
-	AutoUVSettings.Stretch = 0.1;
+	AutoUVSettings.Method = EAutoUVMethod::PatchBuilder;
+	AutoUVSettings.UVAtlasNumCharts = 20;
+	AutoUVSettings.UVAtlasStretch = 0.1;
+	AutoUVSettings.XAtlasMaxIterations = 1;
+	AutoUVSettings.NumInitialPatches = 100;
+	AutoUVSettings.CurvatureAlignment = 1.0;
+	AutoUVSettings.MergingThreshold = 1.5;
+	AutoUVSettings.MaxAngleDeviationDeg = 45.0;
+	AutoUVSettings.SmoothingSteps = 5;
+	AutoUVSettings.SmoothingAlpha = 0.25;
+	AutoUVSettings.bAutoPack = false;
+	AutoUVSettings.PackingTargetWidth = 512;
 	UpdateAutoUVSettings(AutoUVSettings);
 
-	FMeshRecalculateUVsSettings RecomputeUVSettings;
-	UpdateSettingsSourceNodeValue(*Graph, RecomputeUVSettingsNode, RecomputeUVSettings);
+	//FMeshRecalculateUVsSettings RecomputeUVSettings;
+	//UpdateSettingsSourceNodeValue(*Graph, RecomputeUVSettingsNode, RecomputeUVSettings);
 
 	FMeshRepackUVsSettings RepackUVSettings;
 	UpdateSettingsSourceNodeValue(*Graph, RepackUVSettingsNode, RepackUVSettings);
diff --git a/Engine/Plugins/Experimental/MeshLODToolset/Source/MeshLODToolset/Private/Graphs/GenerateStaticMeshLODProcess.cpp b/Engine/Plugins/Experimental/MeshLODToolset/Source/MeshLODToolset/Private/Graphs/GenerateStaticMeshLODProcess.cpp
--- a/Engine/Plugins/Experimental/MeshLODToolset/Source/MeshLODToolset/Private/Graphs/GenerateStaticMeshLODProcess.cpp
+++ b/Engine/Plugins/Experimental/MeshLODToolset/Source/MeshLODToolset/Private/Graphs/GenerateStaticMeshLODProcess.cpp
@@ -553,7 +553,7 @@ bool UGenerateStaticMeshLODProcess::InitializeGenerator()
 
 	// read back default settings
 
-	CurrentSettings.FilterGroupLayer = Generator->GetCurrentPreFilterSettings().FilterGroupLayerName;
+	CurrentSettings_Preprocess.FilterGroupLayer = Generator->GetCurrentPreFilterSettings().FilterGroupLayerName;
 
 	CurrentSettings.SolidifyVoxelResolution = Generator->GetCurrentSolidifySettings().VoxelResolution;
 	CurrentSettings.WindingThreshold = Generator->GetCurrentSolidifySettings().WindingThreshold;
@@ -561,22 +561,32 @@ bool UGenerateStaticMeshLODProcess::InitializeGenerator()
 	//CurrentSettings.MorphologyVoxelResolution = Generator->GetCurrentMorphologySettings().VoxelResolution;
 	CurrentSettings.ClosureDistance = Generator->GetCurrentMorphologySettings().Distance;
 
-	CurrentSettings.SimplifyTriangleCount = Generator->GetCurrentSimplifySettings().TargetCount;
+	CurrentSettings_Simplify.Method = static_cast<EGenerateStaticMeshLODProcess_SimplifyMethod>(Generator->GetCurrentSimplifySettings().TargetType);
+	CurrentSettings_Simplify.TargetCount = Generator->GetCurrentSimplifySettings().TargetCount;
+	CurrentSettings_Simplify.TargetPercentage = Generator->GetCurrentSimplifySettings().TargetFraction * 100.0f;
+	CurrentSettings_Simplify.Tolerance = Generator->GetCurrentSimplifySettings().GeometricTolerance;
 
-	CurrentSettings.BakeResolution = (EGenerateStaticMeshLODBakeResolution)Generator->GetCurrentBakeCacheSettings().Dimensions.GetWidth();
-	CurrentSettings.BakeThickness = Generator->GetCurrentBakeCacheSettings().Thickness;
+	CurrentSettings_Texture.BakeResolution = (EGenerateStaticMeshLODBakeResolution)Generator->GetCurrentBakeCacheSettings().Dimensions.GetWidth();
+	CurrentSettings_Texture.BakeThickness = Generator->GetCurrentBakeCacheSettings().Thickness;
 
-	
-	const UE::GeometryFlow::FGenerateSimpleCollisionSettings& SimpleCollisionSettings = Generator->GetCurrentGenerateSimpleCollisionSettings();
-	CurrentSettings.CollisionType = static_cast<EGenerateStaticMeshLODSimpleCollisionGeometryType>(SimpleCollisionSettings.Type);
-	CurrentSettings.ConvexTriangleCount = SimpleCollisionSettings.ConvexHullSettings.SimplifyToTriangleCount;
-	CurrentSettings.bPrefilterVertices = SimpleCollisionSettings.ConvexHullSettings.bPrefilterVertices;
-	CurrentSettings.PrefilterGridResolution = SimpleCollisionSettings.ConvexHullSettings.PrefilterGridResolution;
-	CurrentSettings.bSimplifyPolygons = SimpleCollisionSettings.SweptHullSettings.bSimplifyPolygons;
-	CurrentSettings.HullTolerance = SimpleCollisionSettings.SweptHullSettings.HullTolerance;
+	CurrentSettings_UV.UVMethod = static_cast<EGenerateStaticMeshLODProcess_AutoUVMethod>(Generator->GetCurrentAutoUVSettings().Method);
+	CurrentSettings_UV.NumUVAtlasCharts = Generator->GetCurrentAutoUVSettings().UVAtlasNumCharts;
+	CurrentSettings_UV.NumInitialPatches = Generator->GetCurrentAutoUVSettings().NumInitialPatches;
+	CurrentSettings_UV.MergingThreshold = Generator->GetCurrentAutoUVSettings().MergingThreshold;
+	CurrentSettings_UV.MaxAngleDeviation = Generator->GetCurrentAutoUVSettings().MaxAngleDeviationDeg;
+	CurrentSettings_UV.PatchBuilder.CurvatureAlignment = Generator->GetCurrentAutoUVSettings().CurvatureAlignment;
+	CurrentSettings_UV.PatchBuilder.SmoothingSteps = Generator->GetCurrentAutoUVSettings().SmoothingSteps;
+	CurrentSettings_UV.PatchBuilder.SmoothingAlpha = Generator->GetCurrentAutoUVSettings().SmoothingAlpha;
 
+	const UE::GeometryFlow::FGenerateSimpleCollisionSettings& SimpleCollisionSettings = Generator->GetCurrentGenerateSimpleCollisionSettings();
+	CurrentSettings_Collision.CollisionType = static_cast<EGenerateStaticMeshLODSimpleCollisionGeometryType>(SimpleCollisionSettings.Type);
+	CurrentSettings_Collision.ConvexTriangleCount = SimpleCollisionSettings.ConvexHullSettings.SimplifyToTriangleCount;
+	CurrentSettings_Collision.bPrefilterVertices = SimpleCollisionSettings.ConvexHullSettings.bPrefilterVertices;
+	CurrentSettings_Collision.PrefilterGridResolution = SimpleCollisionSettings.ConvexHullSettings.PrefilterGridResolution;
+	CurrentSettings_Collision.bSimplifyPolygons = SimpleCollisionSettings.SweptHullSettings.bSimplifyPolygons;
+	CurrentSettings_Collision.HullTolerance = SimpleCollisionSettings.SweptHullSettings.HullTolerance;
 	FMeshSimpleShapeApproximation::EProjectedHullAxisMode RHSMode = SimpleCollisionSettings.SweptHullSettings.SweepAxis;
-	CurrentSettings.SweepAxis = static_cast<EGenerateStaticMeshLODProjectedHullAxisMode>(RHSMode);
+	CurrentSettings_Collision.SweepAxis = static_cast<EGenerateStaticMeshLODProjectedHullAxisMode>(RHSMode);
 
 
 	return true;
@@ -585,22 +595,48 @@ bool UGenerateStaticMeshLODProcess::InitializeGenerator()
 
 void UGenerateStaticMeshLODProcess::UpdateSettings(const FGenerateStaticMeshLODProcessSettings& NewCombinedSettings)
 {
+	bool bSharedVoxelResolutionChanged = (NewCombinedSettings.SolidifyVoxelResolution != CurrentSettings.SolidifyVoxelResolution);
+	if ( bSharedVoxelResolutionChanged
+		|| (NewCombinedSettings.WindingThreshold != CurrentSettings.WindingThreshold))
+	{
+		UE::GeometryFlow::FMeshSolidifySettings NewSolidifySettings = Generator->GetCurrentSolidifySettings();
+		NewSolidifySettings.VoxelResolution = NewCombinedSettings.SolidifyVoxelResolution;
+		NewSolidifySettings.WindingThreshold = NewCombinedSettings.WindingThreshold;
+		Generator->UpdateSolidifySettings(NewSolidifySettings);
+	}
+
+
+	if ( bSharedVoxelResolutionChanged
+		|| (NewCombinedSettings.ClosureDistance != CurrentSettings.ClosureDistance))
+	{
+		UE::GeometryFlow::FVoxClosureSettings NewClosureSettings = Generator->GetCurrentMorphologySettings();
+		NewClosureSettings.VoxelResolution = NewCombinedSettings.SolidifyVoxelResolution;
+		NewClosureSettings.Distance = NewCombinedSettings.ClosureDistance;
+		Generator->UpdateMorphologySettings(NewClosureSettings);
+	}
+
+	CurrentSettings = NewCombinedSettings;
+}
+
+
 
-	if (NewCombinedSettings.FilterGroupLayer != CurrentSettings.FilterGroupLayer)
+void UGenerateStaticMeshLODProcess::UpdatePreprocessSettings(const FGenerateStaticMeshLODProcess_PreprocessSettings& NewCombinedSettings)
+{
+	if (NewCombinedSettings.FilterGroupLayer != CurrentSettings_Preprocess.FilterGroupLayer)
 	{
 		FMeshLODGraphPreFilterSettings NewPreFilterSettings = Generator->GetCurrentPreFilterSettings();
 		NewPreFilterSettings.FilterGroupLayerName = NewCombinedSettings.FilterGroupLayer;
 		Generator->UpdatePreFilterSettings(NewPreFilterSettings);
 	}
 
-	if (NewCombinedSettings.ThickenAmount != CurrentSettings.ThickenAmount)
+	if (NewCombinedSettings.ThickenAmount != CurrentSettings_Preprocess.ThickenAmount)
 	{
 		UE::GeometryFlow::FMeshThickenSettings NewThickenSettings = Generator->GetCurrentThickenSettings();
 		NewThickenSettings.ThickenAmount = NewCombinedSettings.ThickenAmount;
 		Generator->UpdateThickenSettings(NewThickenSettings);
 	}
 
-	if (NewCombinedSettings.ThickenWeightMapName != CurrentSettings.ThickenWeightMapName)
+	if (NewCombinedSettings.ThickenWeightMapName != CurrentSettings_Preprocess.ThickenWeightMapName)
 	{
 		FIndexedWeightMap1f WeightMap;
 		float DefaultValue = 0.0f;
@@ -616,47 +652,36 @@ void UGenerateStaticMeshLODProcess::UpdateSettings(const FGenerateStaticMeshLODP
 		}
 	}
 
-	
-
-	bool bSharedVoxelResolutionChanged = (NewCombinedSettings.SolidifyVoxelResolution != CurrentSettings.SolidifyVoxelResolution);
-	if ( bSharedVoxelResolutionChanged
-		|| (NewCombinedSettings.WindingThreshold != CurrentSettings.WindingThreshold))
-	{
-		UE::GeometryFlow::FMeshSolidifySettings NewSolidifySettings = Generator->GetCurrentSolidifySettings();
-		NewSolidifySettings.VoxelResolution = NewCombinedSettings.SolidifyVoxelResolution;
-		NewSolidifySettings.WindingThreshold = NewCombinedSettings.WindingThreshold;
-		Generator->UpdateSolidifySettings(NewSolidifySettings);
-	}
-
+	CurrentSettings_Preprocess = NewCombinedSettings;
+}
 
-	if ( bSharedVoxelResolutionChanged
-		|| (NewCombinedSettings.ClosureDistance != CurrentSettings.ClosureDistance))
-	{
-		UE::GeometryFlow::FVoxClosureSettings NewClosureSettings = Generator->GetCurrentMorphologySettings();
-		NewClosureSettings.VoxelResolution = NewCombinedSettings.SolidifyVoxelResolution;
-		NewClosureSettings.Distance = NewCombinedSettings.ClosureDistance;
-		Generator->UpdateMorphologySettings(NewClosureSettings);
-	}
 
 
-	if (NewCombinedSettings.SimplifyTriangleCount != CurrentSettings.SimplifyTriangleCount)
+void UGenerateStaticMeshLODProcess::UpdateSimplifySettings(const FGenerateStaticMeshLODProcess_SimplifySettings& NewCombinedSettings)
+{
+	if (NewCombinedSettings.Method != CurrentSettings_Simplify.Method || 
+		NewCombinedSettings.TargetCount != CurrentSettings_Simplify.TargetCount ||
+		NewCombinedSettings.TargetPercentage != CurrentSettings_Simplify.TargetPercentage ||
+		NewCombinedSettings.Tolerance != CurrentSettings_Simplify.Tolerance )
 	{
 		UE::GeometryFlow::FMeshSimplifySettings NewSimplifySettings = Generator->GetCurrentSimplifySettings();
-		NewSimplifySettings.TargetCount = NewCombinedSettings.SimplifyTriangleCount;
+		NewSimplifySettings.TargetType = static_cast<UE::GeometryFlow::EMeshSimplifyTargetType>(NewCombinedSettings.Method);
+		NewSimplifySettings.TargetCount = NewCombinedSettings.TargetCount;
+		NewSimplifySettings.TargetFraction = NewCombinedSettings.TargetPercentage / 100.0f;
+		NewSimplifySettings.GeometricTolerance = NewCombinedSettings.Tolerance;
 		Generator->UpdateSimplifySettings(NewSimplifySettings);
 	}
 
-	if (NewCombinedSettings.NumAutoUVCharts != CurrentSettings.NumAutoUVCharts)
-	{
-		UE::GeometryFlow::FMeshAutoGenerateUVsSettings NewAutoUVSettings = Generator->GetCurrentAutoUVSettings();
-		NewAutoUVSettings.NumCharts = NewCombinedSettings.NumAutoUVCharts;
-		Generator->UpdateAutoUVSettings(NewAutoUVSettings);
-	}
+	CurrentSettings_Simplify = NewCombinedSettings;
+}
 
 
-	if ( (NewCombinedSettings.BakeResolution != CurrentSettings.BakeResolution) ||
-		 (NewCombinedSettings.BakeThickness != CurrentSettings.BakeThickness) || 
-		 (NewCombinedSettings.bCombineTextures != CurrentSettings.bCombineTextures))
+
+void UGenerateStaticMeshLODProcess::UpdateTextureSettings(const FGenerateStaticMeshLODProcess_TextureSettings& NewCombinedSettings)
+{
+	if ( (NewCombinedSettings.BakeResolution != CurrentSettings_Texture.BakeResolution) ||
+		 (NewCombinedSettings.BakeThickness != CurrentSettings_Texture.BakeThickness) || 
+		 (NewCombinedSettings.bCombineTextures != CurrentSettings_Texture.bCombineTextures))
 	{
 		UE::GeometryFlow::FMeshMakeBakingCacheSettings NewBakeSettings = Generator->GetCurrentBakeCacheSettings();
 		NewBakeSettings.Dimensions = FImageDimensions((int32)NewCombinedSettings.BakeResolution, (int32)NewCombinedSettings.BakeResolution);
@@ -664,18 +689,59 @@ void UGenerateStaticMeshLODProcess::UpdateSettings(const FGenerateStaticMeshLODP
 		Generator->UpdateBakeCacheSettings(NewBakeSettings);
 	}
 
-	if (NewCombinedSettings.CollisionGroupLayerName != CurrentSettings.CollisionGroupLayerName)
+	CurrentSettings_Texture = NewCombinedSettings;
+}
+
+
+
+
+void UGenerateStaticMeshLODProcess::UpdateUVSettings(const FGenerateStaticMeshLODProcess_UVSettings& NewCombinedSettings)
+{
+	if (NewCombinedSettings.UVMethod != CurrentSettings_UV.UVMethod
+		|| NewCombinedSettings.NumUVAtlasCharts != CurrentSettings_UV.NumUVAtlasCharts
+		|| NewCombinedSettings.NumInitialPatches != CurrentSettings_UV.NumInitialPatches
+		|| NewCombinedSettings.MergingThreshold != CurrentSettings_UV.MergingThreshold
+		|| NewCombinedSettings.MaxAngleDeviation != CurrentSettings_UV.MaxAngleDeviation
+		|| NewCombinedSettings.PatchBuilder.CurvatureAlignment != CurrentSettings_UV.PatchBuilder.CurvatureAlignment
+		|| NewCombinedSettings.PatchBuilder.SmoothingSteps != CurrentSettings_UV.PatchBuilder.SmoothingSteps
+		|| NewCombinedSettings.PatchBuilder.SmoothingAlpha != CurrentSettings_UV.PatchBuilder.SmoothingAlpha  ) 
+	{
+		UE::GeometryFlow::FMeshAutoGenerateUVsSettings NewAutoUVSettings = Generator->GetCurrentAutoUVSettings();
+		NewAutoUVSettings.Method = (EAutoUVMethod)(int)NewCombinedSettings.UVMethod;
+		//NewAutoUVSettings.UVAtlasStretch = 0.5;
+		NewAutoUVSettings.UVAtlasNumCharts = NewCombinedSettings.NumUVAtlasCharts;
+		//NewAutoUVSettings.XAtlasMaxIterations = 1;
+		NewAutoUVSettings.NumInitialPatches = NewCombinedSettings.NumInitialPatches;
+		NewAutoUVSettings.CurvatureAlignment = NewCombinedSettings.PatchBuilder.CurvatureAlignment;
+		NewAutoUVSettings.MergingThreshold = NewCombinedSettings.MergingThreshold;
+		NewAutoUVSettings.MaxAngleDeviationDeg = NewCombinedSettings.MaxAngleDeviation;
+		NewAutoUVSettings.SmoothingSteps = NewCombinedSettings.PatchBuilder.SmoothingSteps;
+		NewAutoUVSettings.SmoothingAlpha = NewCombinedSettings.PatchBuilder.SmoothingAlpha;
+		//NewAutoUVSettings.bAutoPack = false;
+		//NewAutoUVSettings.PackingTargetWidth = 512;
+		Generator->UpdateAutoUVSettings(NewAutoUVSettings);
+	}
+
+	CurrentSettings_UV = NewCombinedSettings;
+}
+
+
+
+
+void UGenerateStaticMeshLODProcess::UpdateCollisionSettings(const FGenerateStaticMeshLODProcess_CollisionSettings& NewCombinedSettings)
+{
+	if (NewCombinedSettings.CollisionGroupLayerName != CurrentSettings_Collision.CollisionGroupLayerName)
 	{
 		Generator->UpdateCollisionGroupLayerName(NewCombinedSettings.CollisionGroupLayerName);
 	}
 
-	if (NewCombinedSettings.ConvexTriangleCount != CurrentSettings.ConvexTriangleCount ||
-		NewCombinedSettings.bPrefilterVertices != CurrentSettings.bPrefilterVertices ||
-		NewCombinedSettings.PrefilterGridResolution != CurrentSettings.PrefilterGridResolution ||
-		NewCombinedSettings.bSimplifyPolygons != CurrentSettings.bSimplifyPolygons ||
-		NewCombinedSettings.HullTolerance != CurrentSettings.HullTolerance ||
-		NewCombinedSettings.SweepAxis != CurrentSettings.SweepAxis ||
-		NewCombinedSettings.CollisionType != CurrentSettings.CollisionType)
+	if (NewCombinedSettings.ConvexTriangleCount != CurrentSettings_Collision.ConvexTriangleCount ||
+		NewCombinedSettings.bPrefilterVertices != CurrentSettings_Collision.bPrefilterVertices ||
+		NewCombinedSettings.PrefilterGridResolution != CurrentSettings_Collision.PrefilterGridResolution ||
+		NewCombinedSettings.bSimplifyPolygons != CurrentSettings_Collision.bSimplifyPolygons ||
+		NewCombinedSettings.HullTolerance != CurrentSettings_Collision.HullTolerance ||
+		NewCombinedSettings.SweepAxis != CurrentSettings_Collision.SweepAxis ||
+		NewCombinedSettings.CollisionType != CurrentSettings_Collision.CollisionType)
 	{
 		UE::GeometryFlow::FGenerateSimpleCollisionSettings NewGenCollisionSettings = Generator->GetCurrentGenerateSimpleCollisionSettings();
 		NewGenCollisionSettings.Type = static_cast<UE::GeometryFlow::ESimpleCollisionGeometryType>(NewCombinedSettings.CollisionType);
@@ -688,12 +754,11 @@ void UGenerateStaticMeshLODProcess::UpdateSettings(const FGenerateStaticMeshLODP
 		Generator->UpdateGenerateSimpleCollisionSettings(NewGenCollisionSettings);
 	}
 
-	CurrentSettings = NewCombinedSettings;
+	CurrentSettings_Collision = NewCombinedSettings;
 }
 
 
 
-
 bool UGenerateStaticMeshLODProcess::ComputeDerivedSourceData(FProgressCancel* Progress)
 {
 	DerivedTextureImages.Reset();
@@ -782,7 +847,7 @@ void UGenerateStaticMeshLODProcess::GetDerivedMaterialsPreview(FPreviewMaterials
 		{
 			const FTextureInfo& SourceTex = SourceMaterialInfo.SourceTextures[ti];
 
-			if (CurrentSettings.bCombineTextures && SourceTex.bIsUsedInMultiTextureBaking)
+			if (CurrentSettings_Texture.bCombineTextures && SourceTex.bIsUsedInMultiTextureBaking)
 			{
 				continue;
 			}
@@ -815,7 +880,7 @@ void UGenerateStaticMeshLODProcess::GetDerivedMaterialsPreview(FPreviewMaterials
 	MaterialSetOut.Textures.Add(PreviewNormalMapTex);
 
 	// create multi-texture bake result
-	if (CurrentSettings.bCombineTextures)
+	if (CurrentSettings_Texture.bCombineTextures)
 	{
 		bool bConvertToSRGB = true;
 
@@ -848,7 +913,7 @@ void UGenerateStaticMeshLODProcess::GetDerivedMaterialsPreview(FPreviewMaterials
 			// rewrite texture parameters to new textures
 			UpdateMaterialTextureParameters(GeneratedMID, SourceMaterialInfo, SourceToPreviewTexMap, PreviewNormalMapTex);
 
-			if (CurrentSettings.bCombineTextures && DerivedMultiTextureBakeResult && MultiTextureParameterName.Contains(mi))
+			if (CurrentSettings_Texture.bCombineTextures && DerivedMultiTextureBakeResult && MultiTextureParameterName.Contains(mi))
 			{
 				FMaterialParameterInfo ParamInfo(MultiTextureParameterName[mi]);
 				GeneratedMID->SetTextureParameterValueByInfo(ParamInfo, DerivedMultiTextureBakeResult);
@@ -884,7 +949,7 @@ void UGenerateStaticMeshLODProcess::UpdateMaterialTextureParameters(
 		}
 		else if (SourceTex.bShouldBakeTexture)
 		{
-			if (CurrentSettings.bCombineTextures && SourceTex.bIsUsedInMultiTextureBaking)
+			if (CurrentSettings_Texture.bCombineTextures && SourceTex.bIsUsedInMultiTextureBaking)
 			{
 				continue;
 			}
@@ -1064,7 +1129,7 @@ void UGenerateStaticMeshLODProcess::WriteDerivedTextures(bool bCreatingNewStatic
 
 
 	// write multi-texture bake result
-	if (CurrentSettings.bCombineTextures)
+	if (CurrentSettings_Texture.bCombineTextures)
 	{
 		bool bConvertToSRGB = true;
 		
@@ -1238,7 +1303,7 @@ void UGenerateStaticMeshLODProcess::WriteDerivedMaterials(bool bCreatingNewStati
 		DerivedMaterialInfo.DerivedMaterial.MaterialSlotName = FName(FString::Printf(TEXT("GeneratedMat%d"), mi));
 		DerivedMaterialInfo.DerivedMaterial.ImportedMaterialSlotName = DerivedMaterialInfo.DerivedMaterial.MaterialSlotName;
 
-		if (CurrentSettings.bCombineTextures && DerivedMultiTextureBakeResult && MultiTextureParameterName.Contains(mi))
+		if (CurrentSettings_Texture.bCombineTextures && DerivedMultiTextureBakeResult && MultiTextureParameterName.Contains(mi))
 		{
 			FMaterialParameterInfo ParamInfo(MultiTextureParameterName[mi]);
 			GeneratedMIC->SetTextureParameterValueEditorOnly(ParamInfo, DerivedMultiTextureBakeResult);
diff --git a/Engine/Plugins/Experimental/MeshLODToolset/Source/MeshLODToolset/Private/Tools/GenerateStaticMeshLODAssetTool.cpp b/Engine/Plugins/Experimental/MeshLODToolset/Source/MeshLODToolset/Private/Tools/GenerateStaticMeshLODAssetTool.cpp
--- a/Engine/Plugins/Experimental/MeshLODToolset/Source/MeshLODToolset/Private/Tools/GenerateStaticMeshLODAssetTool.cpp
+++ b/Engine/Plugins/Experimental/MeshLODToolset/Source/MeshLODToolset/Private/Tools/GenerateStaticMeshLODAssetTool.cpp
@@ -70,7 +70,12 @@ namespace GenerateStaticMeshLODAssetLocals
 
 		// Inputs
 		UGenerateStaticMeshLODProcess* GenerateProcess;
-		FGenerateStaticMeshLODProcessSettings GeneratorSettings;
+		FGenerateStaticMeshLODProcess_PreprocessSettings GeneratorSettings_Preprocess;
+		FGenerateStaticMeshLODProcessSettings GeneratorSettings_MeshGeneration;
+		FGenerateStaticMeshLODProcess_SimplifySettings GeneratorSettings_Simplify;
+		FGenerateStaticMeshLODProcess_TextureSettings GeneratorSettings_Texture;
+		FGenerateStaticMeshLODProcess_UVSettings GeneratorSettings_UV;
+		FGenerateStaticMeshLODProcess_CollisionSettings GeneratorSettings_Collision;
 		
 		// Outputs
 		// Inherited: 	TUniquePtr<FDynamicMesh3> ResultMesh;
@@ -92,7 +97,12 @@ namespace GenerateStaticMeshLODAssetLocals
 					return;
 				}
 
-				GenerateProcess->UpdateSettings(GeneratorSettings);
+				GenerateProcess->UpdatePreprocessSettings(GeneratorSettings_Preprocess);
+				GenerateProcess->UpdateSettings(GeneratorSettings_MeshGeneration);
+				GenerateProcess->UpdateSimplifySettings(GeneratorSettings_Simplify);
+				GenerateProcess->UpdateTextureSettings(GeneratorSettings_Texture);
+				GenerateProcess->UpdateUVSettings(GeneratorSettings_UV);
+				GenerateProcess->UpdateCollisionSettings(GeneratorSettings_Collision);
 
 				if (Progress && Progress->Cancelled())
 				{
@@ -139,8 +149,13 @@ namespace GenerateStaticMeshLODAssetLocals
 			check(AutoLODTool);
 			TUniquePtr<FGenerateStaticMeshLODAssetOperatorOp> Op = MakeUnique<FGenerateStaticMeshLODAssetOperatorOp>();		
 			Op->GenerateProcess = AutoLODTool->GenerateProcess;
-			Op->GeneratorSettings = AutoLODTool->BasicProperties->GeneratorSettings;
-			Op->GeneratorSettings.CollisionGroupLayerName = AutoLODTool->BasicProperties->CollisionGroupLayerName;
+			Op->GeneratorSettings_Preprocess = AutoLODTool->BasicProperties->Preprocessing;
+			Op->GeneratorSettings_MeshGeneration = AutoLODTool->BasicProperties->MeshGeneration;
+			Op->GeneratorSettings_Simplify = AutoLODTool->BasicProperties->Simplification;
+			Op->GeneratorSettings_Texture = AutoLODTool->BasicProperties->TextureBaking;
+			Op->GeneratorSettings_UV = AutoLODTool->BasicProperties->UVGeneration;
+			Op->GeneratorSettings_Collision.CollisionGroupLayerName = AutoLODTool->BasicProperties->CollisionGroupLayerName;
+			Op->GeneratorSettings_Collision = AutoLODTool->BasicProperties->SimpleCollision;
 			Op->SetResultTransform(ResultTransform);
 			return Op;
 		}
@@ -237,36 +252,28 @@ void UGenerateStaticMeshLODAssetTool::Setup()
 			EToolMessageLevel::UserError);
 	}
 
-	
+	FString FullPathWithExtension = UEditorAssetLibrary::GetPathNameForLoadedAsset(StaticMesh);
+	OutputProperties = NewObject<UGenerateStaticMeshLODAssetToolOutputProperties>(this);
+	AddToolPropertySource(OutputProperties);
+	OutputProperties->NewAssetName = FPaths::GetBaseFilename(FullPathWithExtension, true);
+	OutputProperties->GeneratedSuffix = TEXT("_AutoLOD");
+	OutputProperties->RestoreProperties(this);
 
 	SlowTask.EnterProgressFrame(1);
 	BasicProperties = NewObject<UGenerateStaticMeshLODAssetToolProperties>(this);
 	AddToolPropertySource(BasicProperties);
+	BasicProperties->Preprocessing = GenerateProcess->GetCurrentPreprocessSettings();
+	BasicProperties->MeshGeneration = GenerateProcess->GetCurrentSettings();
+	BasicProperties->Simplification = GenerateProcess->GetCurrentSimplifySettings();
+	BasicProperties->TextureBaking = GenerateProcess->GetCurrentTextureSettings();
+	BasicProperties->UVGeneration = GenerateProcess->GetCurrentUVSettings();
+	BasicProperties->SimpleCollision = GenerateProcess->GetCurrentCollisionSettings();
+	// if we defer restore to here, then on first run we get the defaults, but afterwards we get the restored values
 	BasicProperties->RestoreProperties(this);
-	FString FullPathWithExtension = UEditorAssetLibrary::GetPathNameForLoadedAsset(StaticMesh);
-	BasicProperties->NewAssetName = FPaths::GetBaseFilename(FullPathWithExtension, true);
-	BasicProperties->GeneratedSuffix = TEXT("_AutoLOD");
-	BasicProperties->GeneratorSettings = GenerateProcess->GetCurrentSettings();
-	BasicProperties->WatchProperty(BasicProperties->GeneratorSettings.FilterGroupLayer, [this](FName) { OnSettingsModified(); });
-	BasicProperties->WatchProperty(BasicProperties->GeneratorSettings.ThickenWeightMapName, [this](FName) { OnSettingsModified(); });
-	BasicProperties->WatchProperty(BasicProperties->GeneratorSettings.ThickenAmount, [this](float) { OnSettingsModified(); });
-	BasicProperties->WatchProperty(BasicProperties->GeneratorSettings.SolidifyVoxelResolution, [this](int) { OnSettingsModified(); });
-	BasicProperties->WatchProperty(BasicProperties->GeneratorSettings.WindingThreshold, [this](float) { OnSettingsModified(); });
-	//BasicProperties->WatchProperty(BasicProperties->GeneratorSettings.MorphologyVoxelResolution, [this](int) { OnSettingsModified(); });
-	BasicProperties->WatchProperty(BasicProperties->GeneratorSettings.ClosureDistance, [this](float) { OnSettingsModified(); });
-	BasicProperties->WatchProperty(BasicProperties->GeneratorSettings.SimplifyTriangleCount, [this](int) { OnSettingsModified(); });
-	BasicProperties->WatchProperty(BasicProperties->GeneratorSettings.NumAutoUVCharts, [this](int) { OnSettingsModified(); });
-	BasicProperties->WatchProperty(BasicProperties->GeneratorSettings.BakeResolution, [this](EGenerateStaticMeshLODBakeResolution) { OnSettingsModified(); });
-	BasicProperties->WatchProperty(BasicProperties->GeneratorSettings.BakeThickness, [this](float) { OnSettingsModified(); });
-	BasicProperties->WatchProperty(BasicProperties->GeneratorSettings.bCombineTextures, [this](bool) { OnSettingsModified(); });
-	BasicProperties->WatchProperty(BasicProperties->GeneratorSettings.CollisionType, [this](EGenerateStaticMeshLODSimpleCollisionGeometryType) { OnSettingsModified(); });
-	BasicProperties->WatchProperty(BasicProperties->GeneratorSettings.ConvexTriangleCount, [this](int) { OnSettingsModified(); });
-	BasicProperties->WatchProperty(BasicProperties->GeneratorSettings.bPrefilterVertices, [this](bool) { OnSettingsModified(); });
-	BasicProperties->WatchProperty(BasicProperties->GeneratorSettings.PrefilterGridResolution, [this](float) { OnSettingsModified(); });
-	BasicProperties->WatchProperty(BasicProperties->GeneratorSettings.bSimplifyPolygons, [this](bool) { OnSettingsModified(); });
-	BasicProperties->WatchProperty(BasicProperties->GeneratorSettings.HullTolerance, [this](float) { OnSettingsModified(); });
-	BasicProperties->WatchProperty(BasicProperties->GeneratorSettings.SweepAxis, [this](EGenerateStaticMeshLODProjectedHullAxisMode) { OnSettingsModified(); });
-	
+
+	// assume this is an Editor-only Tool, so we can rely on this
+	BasicProperties->GetOnModified().AddLambda([this](UObject*, FProperty*) { OnSettingsModified(); });
+
 	// Collision layer name property
 	BasicProperties->WatchProperty(BasicProperties->CollisionGroupLayerName, [this](FName) { OnSettingsModified(); });
 	BasicProperties->InitializeGroupLayers(&(GenerateProcess->GetSourceMesh()));
@@ -353,6 +360,7 @@ void UGenerateStaticMeshLODAssetTool::OnSettingsModified()
 
 void UGenerateStaticMeshLODAssetTool::Shutdown(EToolShutdownType ShutdownType)
 {
+	OutputProperties->SaveProperties(this);
 	BasicProperties->SaveProperties(this);
 	CollisionVizSettings->SaveProperties(this);
 
@@ -361,7 +369,7 @@ void UGenerateStaticMeshLODAssetTool::Shutdown(EToolShutdownType ShutdownType)
 
 	if (ShutdownType == EToolShutdownType::Accept)
 	{
-		if (BasicProperties->OutputMode == EGenerateLODAssetOutputMode::UpdateExistingAsset)
+		if (OutputProperties->OutputMode == EGenerateLODAssetOutputMode::UpdateExistingAsset)
 		{
 			UpdateExistingAsset();
 		}
@@ -415,7 +423,7 @@ void UGenerateStaticMeshLODAssetTool::UpdateCollisionVisualization()
 void UGenerateStaticMeshLODAssetTool::CreateNewAsset()
 {
 	check(PreviewWithBackgroundCompute->HaveValidResult());
-	GenerateProcess->CalculateDerivedPathName(BasicProperties->NewAssetName, BasicProperties->GeneratedSuffix);
+	GenerateProcess->CalculateDerivedPathName(OutputProperties->NewAssetName, OutputProperties->GeneratedSuffix);
 
 	check(GenerateProcess->GraphEvalCriticalSection.TryLock());		// No ops should be running
 	GenerateProcess->WriteDerivedAssetData();
@@ -427,13 +435,13 @@ void UGenerateStaticMeshLODAssetTool::CreateNewAsset()
 void UGenerateStaticMeshLODAssetTool::UpdateExistingAsset()
 {
 	check(PreviewWithBackgroundCompute->HaveValidResult());
-	GenerateProcess->CalculateDerivedPathName(BasicProperties->NewAssetName, BasicProperties->GeneratedSuffix);
+	GenerateProcess->CalculateDerivedPathName(OutputProperties->NewAssetName, OutputProperties->GeneratedSuffix);
 
 	check(GenerateProcess->GraphEvalCriticalSection.TryLock());		// No ops should be running
 
 	// only updated HD source if we have no HD source asset. Otherwise we are overwriting with existing lowpoly LOD0.
 	bool bUpdateHDSource =
-		BasicProperties->bSaveInputAsHiResSource &&
+		OutputProperties->bSaveInputAsHiResSource &&
 		(GenerateProcess->GetSourceStaticMesh()->IsHiResMeshDescriptionValid() == false);
 
 	GenerateProcess->UpdateSourceAsset(bUpdateHDSource);
diff --git a/Engine/Plugins/Experimental/MeshLODToolset/Source/MeshLODToolset/Public/Graphs/GenerateMeshLODGraph.h b/Engine/Plugins/Experimental/MeshLODToolset/Source/MeshLODToolset/Public/Graphs/GenerateMeshLODGraph.h
--- a/Engine/Plugins/Experimental/MeshLODToolset/Source/MeshLODToolset/Public/Graphs/GenerateMeshLODGraph.h
+++ b/Engine/Plugins/Experimental/MeshLODToolset/Source/MeshLODToolset/Public/Graphs/GenerateMeshLODGraph.h
@@ -107,8 +107,8 @@ protected:
 	UE::GeometryFlow::FGraph::FHandle AutoUVSettingsNode;
 	UE::GeometryFlow::FMeshAutoGenerateUVsSettings CurrentAutoUVSettings;
 
-	UE::GeometryFlow::FGraph::FHandle RecomputeUVNode;
-	UE::GeometryFlow::FGraph::FHandle RecomputeUVSettingsNode;
+	//UE::GeometryFlow::FGraph::FHandle RecomputeUVNode;
+	//UE::GeometryFlow::FGraph::FHandle RecomputeUVSettingsNode;
 
 	UE::GeometryFlow::FGraph::FHandle RepackUVNode;
 	UE::GeometryFlow::FGraph::FHandle RepackUVSettingsNode;
diff --git a/Engine/Plugins/Experimental/MeshLODToolset/Source/MeshLODToolset/Public/Graphs/GenerateStaticMeshLODProcess.h b/Engine/Plugins/Experimental/MeshLODToolset/Source/MeshLODToolset/Public/Graphs/GenerateStaticMeshLODProcess.h
--- a/Engine/Plugins/Experimental/MeshLODToolset/Source/MeshLODToolset/Public/Graphs/GenerateStaticMeshLODProcess.h
+++ b/Engine/Plugins/Experimental/MeshLODToolset/Source/MeshLODToolset/Public/Graphs/GenerateStaticMeshLODProcess.h
@@ -18,47 +18,8 @@
 struct FMeshDescription;
 using UE::Geometry::FDynamicMesh3;
 
-UENUM()
-enum class EGenerateStaticMeshLODBakeResolution
-{
-	Resolution16 = 16 UMETA(DisplayName = "16 x 16"),
-	Resolution32 = 32 UMETA(DisplayName = "32 x 32"),
-	Resolution64 = 64 UMETA(DisplayName = "64 x 64"),
-	Resolution128 = 128 UMETA(DisplayName = "128 x 128"),
-	Resolution256 = 256 UMETA(DisplayName = "256 x 256"),
-	Resolution512 = 512 UMETA(DisplayName = "512 x 512"),
-	Resolution1024 = 1024 UMETA(DisplayName = "1024 x 1024"),
-	Resolution2048 = 2048 UMETA(DisplayName = "2048 x 2048"),
-	Resolution4096 = 4096 UMETA(DisplayName = "4096 x 4096"),
-	Resolution8192 = 8192 UMETA(DisplayName = "8192 x 8192")
-};
 
-UENUM()
-enum class EGenerateStaticMeshLODSimpleCollisionGeometryType : uint8
-{
-	// NOTE: This must be kept in sync with ESimpleCollisionGeometryType in GenerateSimpleCollisionNode.h
 
-	AlignedBoxes,
-	OrientedBoxes,
-	MinimalSpheres,
-	Capsules,
-	ConvexHulls,
-	SweptHulls,
-	MinVolume,
-	None
-};
-
-// NOTE: This must be kept in sync with FMeshSimpleShapeApproximation::EProjectedHullAxisMode in MeshSimpleShapeApproximation.h
-
-UENUM()
-enum class EGenerateStaticMeshLODProjectedHullAxisMode : uint8
-{
-	X = 0,
-	Y = 1,
-	Z = 2,
-	SmallestBoxDimension = 3,
-	SmallestVolume = 4
-};
 
 
 USTRUCT()
@@ -66,3 +27,29 @@
 {
 	GENERATED_BODY()
 
+	// Solidify settings
+
+	/** Target number of voxels along the maximum dimension for Solidify operation */
+	UPROPERTY(EditAnywhere, Category = Solidify, meta = (UIMin = "8", UIMax = "1024", ClampMin = "8", ClampMax = "1024", DisplayName="Voxel Resolution"))
+	int SolidifyVoxelResolution = 128;
+
+	/** Winding number threshold to determine what is considered inside the mesh during Solidify */
+	UPROPERTY(EditAnywhere, Category = Solidify, AdvancedDisplay, meta = (UIMin = "0.1", UIMax = ".9", ClampMin = "-10", ClampMax = "10"))
+	float WindingThreshold = 0.5f;
+
+
+	// Morphology settings
+
+	/** Offset distance in the Morpohological Closure operation */
+	UPROPERTY(EditAnywhere, Category = Morphology, meta = (UIMin = "0", UIMax = "100", ClampMin = "0", ClampMax = "1000", DisplayName = "Closure Distance"))
+	float ClosureDistance = 1.0f;
+
+};
+
+
+
+USTRUCT()
+struct FGenerateStaticMeshLODProcess_PreprocessSettings
+{
+	GENERATED_BODY()
+
@@ -81,38 +68,133 @@
 	/** Amount to thicken the mesh prior to Solidifying. The thicken weight map values are multiplied by this value. */
 	UPROPERTY(EditAnywhere, Category = DetailFilter, meta = (UIMin = "0", UIMax = "100", ClampMin = "0", ClampMax = "1000"))
 	float ThickenAmount = 0.0f;
-
-	// Solidify settings
-
-	/** Target number of voxels along the maximum dimension for Solidify operation */
-	UPROPERTY(EditAnywhere, Category = Solidify, meta = (UIMin = "8", UIMax = "1024", ClampMin = "8", ClampMax = "1024", DisplayName="Voxel Resolution"))
-	int SolidifyVoxelResolution = 64;
-
-	/** Winding number threshold to determine what is considered inside the mesh during Solidify */
-	UPROPERTY(EditAnywhere, Category = Solidify, AdvancedDisplay, meta = (UIMin = "0.1", UIMax = ".9", ClampMin = "-10", ClampMax = "10"))
-	float WindingThreshold = 0.5f;
+};
 
 
-	// Morphology settings
-
-	/** Offset distance in the Morpohological Closure operation */
-	UPROPERTY(EditAnywhere, Category = Morphology, meta = (UIMin = "0", UIMax = "100", ClampMin = "0", ClampMax = "1000", DisplayName = "Closure Distance"))
-	float ClosureDistance = 1.0f;
 
 
-	// Simplify settings
-	/** Target triangle count after simplification */
-	UPROPERTY(EditAnywhere, Category = Simplify, meta = (UIMin = "1", ClampMin = "1", DisplayName = "Simplify Tri Count"))
-	int SimplifyTriangleCount = 500;
+UENUM()
+enum class EGenerateStaticMeshLODProcess_SimplifyMethod : uint8
+{
+	// note: must keep in sync with UE::GeometryFlow::EMeshSimplifyTargetType
+	TriangleCount = 0,
+	VertexCount = 1,
+	TrianglePercentage = 2,
+	GeometricTolerance = 3
+};
 
 
-	// UV settings
-	/** Maximum number of charts to create in AutoUV */
-	UPROPERTY(EditAnywhere, Category = AutoUV, meta = (DisplayName = "Max AutoUV Charts", UIMin = 0, UIMax = 1000))
-	int NumAutoUVCharts = 0;
+USTRUCT()
+struct FGenerateStaticMeshLODProcess_SimplifySettings
+{
+	GENERATED_BODY()
+
+	UPROPERTY(EditAnywhere, Category = Simplify, meta = (DisplayName = "Simplification Target"))
+	EGenerateStaticMeshLODProcess_SimplifyMethod Method = EGenerateStaticMeshLODProcess_SimplifyMethod::GeometricTolerance;
+
+	/** Target triangle/vertex count after simplification */
+	UPROPERTY(EditAnywhere, Category = Simplify, meta = (UIMin = "1", ClampMin = "1", EditConditionHides, EditCondition = "Method == EGenerateStaticMeshLODProcess_SimplifyMethod::TriangleCount || Method == EGenerateStaticMeshLODProcess_SimplifyMethod::VertexCount"))
+	int TargetCount = 500;
+
+	/** Target triangle/vertex count after simplification */
+	UPROPERTY(EditAnywhere, Category = Simplify, meta = (UIMin = "0", UIMax = "100", ClampMin = "0", ClampMax = "100", EditConditionHides, EditCondition = "Method == EGenerateStaticMeshLODProcess_SimplifyMethod::TrianglePercentage"))
+	float TargetPercentage = 5.0;
+
+	/** Deviation in World Units (Centimeters) */
+	UPROPERTY(EditAnywhere, Category = Simplify, meta = (UIMin = "0", UIMax = "100", ClampMin = "0", ClampMax = "100", EditConditionHides, EditCondition = "Method == EGenerateStaticMeshLODProcess_SimplifyMethod::GeometricTolerance"))
+	float Tolerance = 0.5;
+};
 
 
-	// Bake Settings
+
+UENUM()
+enum class EGenerateStaticMeshLODProcess_AutoUVMethod : uint8
+{
+	PatchBuilder = 0,
+	UVAtlas = 1,
+	XAtlas = 2
+};
+
+
+USTRUCT()
+struct FGenerateStaticMeshLODProcess_UVSettings_PatchBuilder
+{
+	GENERATED_BODY()
+
+	/** This parameter controls alignment of the initial patches to creases in the mesh */
+	UPROPERTY(EditAnywhere, Category = "PatchBuilder", meta = (UIMin = "0.1", UIMax = "2.0", ClampMin = "0.01", ClampMax = "100.0", EditConditionHides, EditCondition = "UVMethod == EGenerateStaticMeshLODProcess_AutoUVMethod::PatchBuilder"))
+	float CurvatureAlignment = 1.0f;
+
+	/** Number of smoothing steps to apply in the ExpMap UV Generation method (Smoothing slightly increases distortion but produces more stable results) */
+	UPROPERTY(EditAnywhere, Category = "PatchBuilder", AdvancedDisplay, meta = (UIMin = "0", UIMax = "25", ClampMin = "0", ClampMax = "1000", EditConditionHides, EditCondition = "UVMethod == EGenerateStaticMeshLODProcess_AutoUVMethod::PatchBuilder"))
+	int SmoothingSteps = 5;
+
+	/** Smoothing parameter, larger values result in faster smoothing in each step */
+	UPROPERTY(EditAnywhere, Category = "PatchBuilder", AdvancedDisplay, meta = (UIMin = "0", UIMax = "1.0", ClampMin = "0", ClampMax = "1.0", EditConditionHides, EditCondition = "UVMethod == EGenerateStaticMeshLODProcess_AutoUVMethod::PatchBuilder"))
+	float SmoothingAlpha = 0.25f;
+};
+
+USTRUCT()
+struct FGenerateStaticMeshLODProcess_UVSettings
+{
+	GENERATED_BODY()
+
+	UPROPERTY(EditAnywhere, Category = UVGeneration, meta = (DisplayName = "UV Method"))
+	EGenerateStaticMeshLODProcess_AutoUVMethod UVMethod = EGenerateStaticMeshLODProcess_AutoUVMethod::PatchBuilder;
+
+	/** Maximum number of charts to create in UVAtlas */
+	UPROPERTY(EditAnywhere, Category = "UVAtlas", meta = (DisplayName = "Max Charts", UIMin = 0, UIMax = 1000, EditConditionHides, EditCondition = "UVMethod == EGenerateStaticMeshLODProcess_AutoUVMethod::UVAtlas"))
+	int NumUVAtlasCharts = 0;
+
+	/** Number of initial patches mesh will be split into before computing island merging */
+	UPROPERTY(EditAnywhere, Category = "PatchBuilder", meta = (UIMin = "1", UIMax = "1000", ClampMin = "1", ClampMax = "99999999", EditConditionHides, EditCondition = "UVMethod == EGenerateStaticMeshLODProcess_AutoUVMethod::PatchBuilder"))
+	int NumInitialPatches = 100;
+
+	/** Distortion/Stretching Threshold for island merging - larger values increase the allowable UV stretching */
+	UPROPERTY(EditAnywhere, Category = "PatchBuilder", meta = (UIMin = "1.0", UIMax = "5.0", ClampMin = "1.0", EditConditionHides, EditCondition = "UVMethod == EGenerateStaticMeshLODProcess_AutoUVMethod::PatchBuilder"))
+	float MergingThreshold = 1.5f;
+
+	/** UV islands will not be merged if their average face normals deviate by larger than this amount */
+	UPROPERTY(EditAnywhere, Category = "PatchBuilder", meta = (UIMin = "0.0", UIMax = "90.0", ClampMin = "0.0", ClampMax = "180.0", EditConditionHides, EditCondition = "UVMethod == EGenerateStaticMeshLODProcess_AutoUVMethod::PatchBuilder"))
+	float MaxAngleDeviation = 45.0f;
+
+	UPROPERTY(EditAnywhere, Category = "PatchBuilder", meta = (DisplayName = "UV - PatchBuilder", UIMin = "0.0", UIMax = "90.0", ClampMin = "0.0", ClampMax = "180.0", EditConditionHides, EditCondition = "UVMethod == EGenerateStaticMeshLODProcess_AutoUVMethod::PatchBuilder"))
+	FGenerateStaticMeshLODProcess_UVSettings_PatchBuilder PatchBuilder;
+
+	///** Number of smoothing steps to apply in the ExpMap UV Generation method (Smoothing slightly increases distortion but produces more stable results) */
+	//UPROPERTY(EditAnywhere, Category = "PatchBuilder", AdvancedDisplay, meta = (UIMin = "0", UIMax = "25", ClampMin = "0", ClampMax = "1000", EditConditionHides, EditCondition = "UVMethod == EGenerateStaticMeshLODProcess_AutoUVMethod::PatchBuilder"))
+	//int SmoothingSteps = 5;
+
+	///** Smoothing parameter, larger values result in faster smoothing in each step */
+	//UPROPERTY(EditAnywhere, Category = "PatchBuilder", AdvancedDisplay, meta = (UIMin = "0", UIMax = "1.0", ClampMin = "0", ClampMax = "1.0", EditConditionHides, EditCondition = "UVMethod == EGenerateStaticMeshLODProcess_AutoUVMethod::PatchBuilder"))
+	//float SmoothingAlpha = 0.25f;
+};
+
+
+
+
+UENUM()
+enum class EGenerateStaticMeshLODBakeResolution
+{
+	Resolution16 = 16 UMETA(DisplayName = "16 x 16"),
+	Resolution32 = 32 UMETA(DisplayName = "32 x 32"),
+	Resolution64 = 64 UMETA(DisplayName = "64 x 64"),
+	Resolution128 = 128 UMETA(DisplayName = "128 x 128"),
+	Resolution256 = 256 UMETA(DisplayName = "256 x 256"),
+	Resolution512 = 512 UMETA(DisplayName = "512 x 512"),
+	Resolution1024 = 1024 UMETA(DisplayName = "1024 x 1024"),
+	Resolution2048 = 2048 UMETA(DisplayName = "2048 x 2048"),
+	Resolution4096 = 4096 UMETA(DisplayName = "4096 x 4096"),
+	Resolution8192 = 8192 UMETA(DisplayName = "8192 x 8192")
+};
+
+
+
+USTRUCT()
+struct FGenerateStaticMeshLODProcess_TextureSettings
+{
+	GENERATED_BODY()
+
 	/** Resolution for normal map and texture baking */
 	UPROPERTY(EditAnywhere, Category = Baking , meta = (DisplayName = "Bake Image Res"))
 	EGenerateStaticMeshLODBakeResolution BakeResolution = EGenerateStaticMeshLODBakeResolution::Resolution512;
@@ -123,9 +205,44 @@
 
 	UPROPERTY(EditAnywhere, Category = Baking, meta = (DisplayName = "Combine Textures"))
 	bool bCombineTextures = true;
+};
 
 
-	// Simple collision generator settings
 
+
+UENUM()
+enum class EGenerateStaticMeshLODSimpleCollisionGeometryType : uint8
+{
+	// NOTE: This must be kept in sync with ESimpleCollisionGeometryType in GenerateSimpleCollisionNode.h
+
+	AlignedBoxes,
+	OrientedBoxes,
+	MinimalSpheres,
+	Capsules,
+	ConvexHulls,
+	SweptHulls,
+	MinVolume,
+	None
+};
+
+// NOTE: This must be kept in sync with FMeshSimpleShapeApproximation::EProjectedHullAxisMode in MeshSimpleShapeApproximation.h
+
+UENUM()
+enum class EGenerateStaticMeshLODProjectedHullAxisMode : uint8
+{
+	X = 0,
+	Y = 1,
+	Z = 2,
+	SmallestBoxDimension = 3,
+	SmallestVolume = 4
+};
+
+
+
+USTRUCT()
+struct FGenerateStaticMeshLODProcess_CollisionSettings
+{
+	GENERATED_BODY()
+
 	// Transient property, not set directly by the user. The user controls a CollisionGroupLayerName dropdown property
 	// on the Tool and that value is copied here.
@@ -137,6 +254,7 @@
 	EGenerateStaticMeshLODSimpleCollisionGeometryType CollisionType = EGenerateStaticMeshLODSimpleCollisionGeometryType::ConvexHulls;
 
 	// Convex Hull Settings
+
 	/** Target triangle count for each convex hull after simplification */
 	UPROPERTY(EditAnywhere, Category = Collision, meta = (NoSpinbox = "true", DisplayName = "Convex Tri Count",
 														  EditCondition = "CollisionType == EGenerateStaticMeshLODSimpleCollisionGeometryType::ConvexHulls"))
@@ -152,6 +270,7 @@
 
 
 	// Swept Convex Hull Settings
+
 	/** Whether to simplify polygons used for swept convex hulls */
 	UPROPERTY(EditAnywhere, Category = Collision, meta = (EditCondition = "CollisionType == EGenerateStaticMeshLODSimpleCollisionGeometryType::SweptHulls"))
 	bool bSimplifyPolygons = true;
@@ -164,10 +283,10 @@
 	/** Which axis to sweep along when computing swept convex hulls */
 	UPROPERTY(EditAnywhere, Category = Collision, meta = (EditCondition = "CollisionType == EGenerateStaticMeshLODSimpleCollisionGeometryType::SweptHulls"))
 	EGenerateStaticMeshLODProjectedHullAxisMode SweepAxis = EGenerateStaticMeshLODProjectedHullAxisMode::SmallestVolume;
-
 };
 
 
+
 UCLASS(Transient)
 class UGenerateStaticMeshLODProcess : public UObject
 {
@@ -176,9 +295,25 @@
 
 	bool Initialize(UStaticMesh* SourceMesh, FProgressCancel* Progress = nullptr);
 
+	const FGenerateStaticMeshLODProcess_PreprocessSettings& GetCurrentPreprocessSettings() const { return CurrentSettings_Preprocess; }
+	void UpdatePreprocessSettings(const FGenerateStaticMeshLODProcess_PreprocessSettings& NewSettings);
+
 	const FGenerateStaticMeshLODProcessSettings& GetCurrentSettings() const { return CurrentSettings; }
 	void UpdateSettings(const FGenerateStaticMeshLODProcessSettings& NewSettings);
 
+	const FGenerateStaticMeshLODProcess_SimplifySettings& GetCurrentSimplifySettings() const { return CurrentSettings_Simplify; }
+	void UpdateSimplifySettings(const FGenerateStaticMeshLODProcess_SimplifySettings& NewSettings);
 
+	const FGenerateStaticMeshLODProcess_TextureSettings& GetCurrentTextureSettings() const { return CurrentSettings_Texture; }
+	void UpdateTextureSettings(const FGenerateStaticMeshLODProcess_TextureSettings& NewSettings);
+
+	const FGenerateStaticMeshLODProcess_UVSettings& GetCurrentUVSettings() const { return CurrentSettings_UV; }
+	void UpdateUVSettings(const FGenerateStaticMeshLODProcess_UVSettings& NewSettings);
+
+	const FGenerateStaticMeshLODProcess_CollisionSettings& GetCurrentCollisionSettings() const { return CurrentSettings_Collision; }
+	void UpdateCollisionSettings(const FGenerateStaticMeshLODProcess_CollisionSettings& NewSettings);
+
+
+
 	UStaticMesh* GetSourceStaticMesh() const { return SourceStaticMesh; }
 	const FDynamicMesh3& GetSourceMesh() const { return SourceMesh; }
@@ -315,7 +450,13 @@
 	TUniquePtr<FGenerateMeshLODGraph> Generator;			// active LODGenerator Graph
 	bool InitializeGenerator();
 
+	FGenerateStaticMeshLODProcess_PreprocessSettings CurrentSettings_Preprocess;
 	FGenerateStaticMeshLODProcessSettings CurrentSettings;
+	FGenerateStaticMeshLODProcess_SimplifySettings CurrentSettings_Simplify;
+	FGenerateStaticMeshLODProcess_TextureSettings CurrentSettings_Texture;
+	FGenerateStaticMeshLODProcess_UVSettings CurrentSettings_UV;
+	FGenerateStaticMeshLODProcess_CollisionSettings CurrentSettings_Collision;
+	
 
 	bool WriteDerivedTexture(UTexture2D* SourceTexture, UTexture2D* DerivedTexture, bool bCreatingNewStaticMeshAsset);
 	bool WriteDerivedTexture(UTexture2D* DerivedTexture, FString BaseTexName, bool bCreatingNewStaticMeshAsset);
diff --git a/Engine/Plugins/Experimental/MeshLODToolset/Source/MeshLODToolset/Public/Tools/GenerateStaticMeshLODAssetTool.h b/Engine/Plugins/Experimental/MeshLODToolset/Source/MeshLODToolset/Public/Tools/GenerateStaticMeshLODAssetTool.h
--- a/Engine/Plugins/Experimental/MeshLODToolset/Source/MeshLODToolset/Public/Tools/GenerateStaticMeshLODAssetTool.h
+++ b/Engine/Plugins/Experimental/MeshLODToolset/Source/MeshLODToolset/Public/Tools/GenerateStaticMeshLODAssetTool.h
@@ -50,16 +50,12 @@ protected:
 };
 
 
-/**
- * Standard properties
- */
+
 UCLASS()
-class MESHLODTOOLSET_API UGenerateStaticMeshLODAssetToolProperties : public UInteractiveToolPropertySet
+class MESHLODTOOLSET_API UGenerateStaticMeshLODAssetToolOutputProperties : public UInteractiveToolPropertySet
 {
 	GENERATED_BODY()
-
 public:
-
 	/** Whether to modify the static mesh in place or create a new one. */
 	UPROPERTY(EditAnywhere, Category = AssetOptions, meta = (TransientToolProperty))
 	EGenerateLODAssetOutputMode OutputMode = EGenerateLODAssetOutputMode::UpdateExistingAsset;
@@ -75,9 +71,36 @@ public:
 	/** If this is checked and the asset doesn't already have a HiRes source, the input mesh will be stored as the HiRes source */
 	UPROPERTY(EditAnywhere, Category = AssetOptions, meta = (EditCondition = "OutputMode == EGenerateLODAssetOutputMode::UpdateExistingAsset"))
 	bool bSaveInputAsHiResSource = true;
+};
+
+
+/**
+ * Standard properties
+ */
+UCLASS()
+class MESHLODTOOLSET_API UGenerateStaticMeshLODAssetToolProperties : public UInteractiveToolPropertySet
+{
+	GENERATED_BODY()
+
+public:
+	UPROPERTY(EditAnywhere, Category = Settings)
+	FGenerateStaticMeshLODProcess_PreprocessSettings Preprocessing;
+
+	UPROPERTY(EditAnywhere, Category = Settings)
+	FGenerateStaticMeshLODProcessSettings MeshGeneration;
+
+	UPROPERTY(EditAnywhere, Category = Settings)
+	FGenerateStaticMeshLODProcess_SimplifySettings Simplification;
 
 	UPROPERTY(EditAnywhere, Category = Settings)
-	FGenerateStaticMeshLODProcessSettings GeneratorSettings;
+	FGenerateStaticMeshLODProcess_TextureSettings TextureBaking;
+
+	UPROPERTY(EditAnywhere, Category = Settings)
+	FGenerateStaticMeshLODProcess_UVSettings UVGeneration;
+
+	UPROPERTY(EditAnywhere, Category = Settings)
+	FGenerateStaticMeshLODProcess_CollisionSettings SimpleCollision;
+
 
 	// ------------
 	// Code copied from UPolygroupLayersProperties
@@ -152,6 +175,9 @@ protected:
 
 	friend class GenerateStaticMeshLODAssetLocals::FGenerateStaticMeshLODAssetOperatorFactory;
 
+	UPROPERTY()
+	UGenerateStaticMeshLODAssetToolOutputProperties* OutputProperties;
+
 	UPROPERTY()
 	UGenerateStaticMeshLODAssetToolProperties* BasicProperties;
 
