AutoLOD:
- add support for a Settings Preset, as an explicit Asset type. UStaticMeshLODGenerationSettings contains a copy of the GenerateProcess Settings, UStaticMeshLODGenerationSettingsFactory allows this UObject to be used as an Asset in the Editor.  UGenerateStaticMeshLODAssetTool now has a Preset slot, and action buttons to Read/Write to the currently-selected Settings Asset. New Settings Asset can be created via the Asset Picker popup in the Tool, currently creating via the New Asset menu is disabled via UStaticMeshLODGenerationSettingsFactory::GetMenuCategories()
- UGenerateStaticMeshLODProcess now automatically ignores input textures that have a constant value
- UGenerateStaticMeshLODProcess now skips writing output NormalMap and MultiTexture if they are not referenced by any output Materials
- UGenerateStaticMeshLODProcess now supports configuring whether a Texture or Material will be considered for Baking. If disabled, then the source Texture/Material is used (currently the Textures are still baked, they are just ignored while configuring/writing the outputs).
- UGenerateStaticMeshLODAssetTool now exposes lists of source Texture and Materials with option to skip baking (used to configure the above options)
- wrote header comments for UGenerateStaticMeshLODProcess
- overall cleanup in  UGenerateStaticMeshLODAssetTool, ie category names, organization, etc

- add TImageBuilder::IsConstantValue(), checks if all pixels of image are the same value

#rnx

diff --git a/Engine/Plugins/Experimental/MeshLODToolset/Source/MeshLODToolset/Private/Graphs/GenerateStaticMeshLODProcess.cpp b/Engine/Plugins/Experimental/MeshLODToolset/Source/MeshLODToolset/Private/Graphs/GenerateStaticMeshLODProcess.cpp
--- a/Engine/Plugins/Experimental/MeshLODToolset/Source/MeshLODToolset/Private/Graphs/GenerateStaticMeshLODProcess.cpp
+++ b/Engine/Plugins/Experimental/MeshLODToolset/Source/MeshLODToolset/Private/Graphs/GenerateStaticMeshLODProcess.cpp
@@ -380,6 +380,8 @@ bool UGenerateStaticMeshLODProcess::Initialize(UStaticMesh* StaticMeshIn, FProgr
 	//});
 
 	// single-thread path
+	TArray<UTexture2D*> UniqueTextures;
+	TArray<TImageBuilder<FVector4f>*> UniqueTextureImages;
 	for (FReadTextureJob job : ReadJobs)
 	{
 		// only read textures that are from materials we are going to possibly bake
@@ -389,20 +391,58 @@ bool UGenerateStaticMeshLODProcess::Initialize(UStaticMesh* StaticMeshIn, FProgr
 			FTextureInfo& TexInfo = SourceMaterials[job.MatIndex].SourceTextures[job.TexIndex];
 			UE::AssetUtils::ReadTexture(TexInfo.Texture, TexInfo.Image);
 			TexInfo.Dimensions = TexInfo.Image.GetDimensions();
+
+			UniqueTextures.AddUnique(TexInfo.Texture);
+			UniqueTextureImages.Add(&TexInfo.Image);
 		}
 	}
 
-
-	ConvertMesh.Wait();
 	//ReadTextures.Wait();
 
 
+	// Now that we have read textures, detect solid-color textures. We do not
+	// have to bake these. These are often used for (eg) default textures, etc.
+	TArray<bool> bIsSolidColor;
+	bIsSolidColor.Init(false, UniqueTextures.Num());
+	ParallelFor(UniqueTextures.Num(), [&](int32 ti)
+	{
+		bIsSolidColor[ti] = UniqueTextureImages[ti]->IsConstantValue();
+	});
+	for (FSourceMaterialInfo& SourceMaterial : SourceMaterials)
+	{
+		if (SourceMaterial.bHasTexturesToBake)
+		{
+			int StillToBakeCount = 0;
+			for (FTextureInfo& TexInfo : SourceMaterial.SourceTextures)
+			{
+				int32 Index = UniqueTextures.IndexOfByKey(TexInfo.Texture);
+				if (Index != INDEX_NONE && bIsSolidColor[Index])
+				{
+					TexInfo.bShouldBakeTexture = false;
+				}
+				else
+				{
+					StillToBakeCount++;
+				}
+			}
+			if (StillToBakeCount == 0)
+			{
+				// If Material don't have any remaining textures to bake, and it doesn't have a normal map, then we do not have to bake it
+				SourceMaterial.bHasTexturesToBake = false;
+				SourceMaterial.bIsReusable = (SourceMaterial.bHasNormalMap == false);
+			}
+		}
+	}
+
+
+	ConvertMesh.Wait();
+
 	FString FullPathWithExtension = UEditorAssetLibrary::GetPathNameForLoadedAsset(SourceStaticMesh);
 	SourceAssetPath = FPaths::GetBaseFilename(FullPathWithExtension, false);
 	SourceAssetFolder = FPaths::GetPath(SourceAssetPath);
 	SourceAssetName = FPaths::GetBaseFilename(FullPathWithExtension, true);
 
-	CalculateDerivedPathName(SourceAssetName, GetDefaultDerivedAssetSuffix());
+	UpdateDerivedPathName(SourceAssetName, GetDefaultDerivedAssetSuffix());
 
 	InitializeGenerator();
 
@@ -410,7 +450,7 @@ bool UGenerateStaticMeshLODProcess::Initialize(UStaticMesh* StaticMeshIn, FProgr
 }
 
 
-void UGenerateStaticMeshLODProcess::CalculateDerivedPathName(const FString& NewAssetBaseName, const FString& NewAssetSuffix)
+void UGenerateStaticMeshLODProcess::UpdateDerivedPathName(const FString& NewAssetBaseName, const FString& NewAssetSuffix)
 {
 	DerivedAssetNameNoSuffix = FPaths::MakeValidFileName(NewAssetBaseName);
 	if (DerivedAssetNameNoSuffix.Len() == 0)
@@ -759,6 +799,72 @@ void UGenerateStaticMeshLODProcess::UpdateCollisionSettings(const FGenerateStati
 
 
 
+TArray<UMaterialInterface*> UGenerateStaticMeshLODProcess::GetSourceBakeMaterials() const
+{
+	TArray<UMaterialInterface*> UniqueMaterials;
+	for (const FSourceMaterialInfo& MatInfo : SourceMaterials)
+	{
+		if (MatInfo.bHasTexturesToBake)
+		{
+			UniqueMaterials.AddUnique(MatInfo.SourceMaterial.MaterialInterface);
+		}
+	}
+	return UniqueMaterials;
+}
+
+void UGenerateStaticMeshLODProcess::UpdateSourceBakeMaterialConstraint(UMaterialInterface* Material, EMaterialBakingConstraint Constraint)
+{
+	for (FSourceMaterialInfo& MatInfo : SourceMaterials)
+	{
+		if (MatInfo.SourceMaterial.MaterialInterface == Material)
+		{
+			MatInfo.Constraint = Constraint;
+		}
+	}
+}
+
+
+TArray<UTexture2D*> UGenerateStaticMeshLODProcess::GetSourceBakeTextures() const
+{
+	TArray<UTexture2D*> UniqueTextures;
+	for (const FSourceMaterialInfo& MatInfo : SourceMaterials)
+	{
+		if (MatInfo.bHasTexturesToBake)
+		{
+			for (const FTextureInfo& TexInfo : MatInfo.SourceTextures)
+			{
+				if (TexInfo.bShouldBakeTexture)
+				{
+					UniqueTextures.AddUnique(TexInfo.Texture);
+				}
+			}
+		}
+	}
+	return UniqueTextures;
+}
+
+
+void UGenerateStaticMeshLODProcess::UpdateSourceBakeTextureConstraint(UTexture2D* Texture, ETextureBakingConstraint Constraint)
+{
+	for (FSourceMaterialInfo& MatInfo : SourceMaterials)
+	{
+		if (MatInfo.bHasTexturesToBake)
+		{
+			for (FTextureInfo& TexInfo : MatInfo.SourceTextures)
+			{
+				if (TexInfo.Texture == Texture)
+				{
+					TexInfo.Constraint = Constraint;
+				}
+			}
+		}
+	}
+}
+
+
+
+
+
 bool UGenerateStaticMeshLODProcess::ComputeDerivedSourceData(FProgressCancel* Progress)
 {
 	DerivedTextureImages.Reset();
@@ -783,7 +889,10 @@ bool UGenerateStaticMeshLODProcess::ComputeDerivedSourceData(FProgressCancel* Pr
 	for (int32 mi = 0; mi < NumMaterials; ++mi)
 	{
 		DerivedMaterials[mi].SourceMaterialIndex = mi;
-		DerivedMaterials[mi].bUseSourceMaterialDirectly = (SourceMaterials[mi].bIsReusable || SourceMaterials[mi].bIsPreviouslyGeneratedMaterial);
+		DerivedMaterials[mi].bUseSourceMaterialDirectly = 
+			SourceMaterials[mi].bIsReusable 
+			|| SourceMaterials[mi].bIsPreviouslyGeneratedMaterial
+			|| (SourceMaterials[mi].Constraint == EMaterialBakingConstraint::UseExistingMaterial);
 		
 		if (DerivedMaterials[mi].bUseSourceMaterialDirectly)
 		{
@@ -826,6 +935,7 @@ void UGenerateStaticMeshLODProcess::GetDerivedMaterialsPreview(FPreviewMaterials
 	// force garbage collection of outstanding preview materials
 	CollectGarbage(GARBAGE_COLLECTION_KEEPFLAGS);
 
+	bool bAnyMaterialUsedGeneratedNormalMap = false;
 
 	// create derived textures
 	int32 NumMaterials = SourceMaterials.Num();
@@ -834,11 +944,18 @@ void UGenerateStaticMeshLODProcess::GetDerivedMaterialsPreview(FPreviewMaterials
 	for (int32 mi = 0; mi < NumMaterials; ++mi)
 	{
 		const FSourceMaterialInfo& SourceMaterialInfo = SourceMaterials[mi];
-		if (SourceMaterialInfo.bIsReusable || SourceMaterialInfo.bIsPreviouslyGeneratedMaterial)
+		if (SourceMaterialInfo.bIsReusable 
+			|| SourceMaterialInfo.bIsPreviouslyGeneratedMaterial
+			|| SourceMaterialInfo.Constraint == EMaterialBakingConstraint::UseExistingMaterial )
 		{
 			continue;
 		}
 
+		if (SourceMaterialInfo.bHasNormalMap)
+		{
+			bAnyMaterialUsedGeneratedNormalMap = true;
+		}
+
 		const FDerivedMaterialInfo& DerivedMaterialInfo = DerivedMaterials[mi];
 
 		int32 NumTextures = SourceMaterialInfo.SourceTextures.Num();
@@ -854,7 +971,8 @@ void UGenerateStaticMeshLODProcess::GetDerivedMaterialsPreview(FPreviewMaterials
 
 			bool bConvertToSRGB = SourceTex.Texture->SRGB;
 			const FTextureInfo& DerivedTex = DerivedMaterialInfo.DerivedTextures[ti];
-			if (DerivedTex.bShouldBakeTexture)
+			if (DerivedTex.bShouldBakeTexture
+				&& SourceTex.Constraint != ETextureBakingConstraint::UseExistingTexture )
 			{
 				FTexture2DBuilder TextureBuilder;
 				TextureBuilder.Initialize(FTexture2DBuilder::ETextureType::Color, DerivedTex.Dimensions);
@@ -872,12 +990,16 @@ void UGenerateStaticMeshLODProcess::GetDerivedMaterialsPreview(FPreviewMaterials
 	}
 
 	// create derived normal map texture
-	FTexture2DBuilder NormapMapBuilder;
-	NormapMapBuilder.Initialize(FTexture2DBuilder::ETextureType::NormalMap, DerivedNormalMapImage.Image.GetDimensions());
-	NormapMapBuilder.Copy(DerivedNormalMapImage.Image, false);
-	NormapMapBuilder.Commit(false);
-	UTexture2D* PreviewNormalMapTex = NormapMapBuilder.GetTexture2D();
-	MaterialSetOut.Textures.Add(PreviewNormalMapTex);
+	UTexture2D* PreviewNormalMapTex = nullptr;
+	if (bAnyMaterialUsedGeneratedNormalMap)
+	{
+		FTexture2DBuilder NormapMapBuilder;
+		NormapMapBuilder.Initialize(FTexture2DBuilder::ETextureType::NormalMap, DerivedNormalMapImage.Image.GetDimensions());
+		NormapMapBuilder.Copy(DerivedNormalMapImage.Image, false);
+		NormapMapBuilder.Commit(false);
+		PreviewNormalMapTex = NormapMapBuilder.GetTexture2D();
+		MaterialSetOut.Textures.Add(PreviewNormalMapTex);
+	}
 
 	// create multi-texture bake result
 	if (CurrentSettings_Texture.bCombineTextures)
@@ -901,7 +1023,9 @@ void UGenerateStaticMeshLODProcess::GetDerivedMaterialsPreview(FPreviewMaterials
 		const FSourceMaterialInfo& SourceMaterialInfo = SourceMaterials[mi];
 		UMaterialInterface* MaterialInterface = SourceMaterialInfo.SourceMaterial.MaterialInterface;
 
-		if (SourceMaterialInfo.bIsReusable || SourceMaterialInfo.bIsPreviouslyGeneratedMaterial)
+		if (SourceMaterialInfo.bIsReusable 
+			|| SourceMaterialInfo.bIsPreviouslyGeneratedMaterial
+			|| SourceMaterialInfo.Constraint == EMaterialBakingConstraint::UseExistingMaterial )
 		{
 			MaterialSetOut.Materials.Add(MaterialInterface);
 		}
@@ -933,15 +1057,22 @@ void UGenerateStaticMeshLODProcess::UpdateMaterialTextureParameters(
 	const TMap<UTexture2D*, UTexture2D*>& PreviewTextures, 
 	UTexture2D* PreviewNormalMap)
 {
+	ensure((SourceMaterialInfo.bHasNormalMap == false) || (PreviewNormalMap != nullptr));
+
 	Material->Modify();
 	int32 NumTextures = SourceMaterialInfo.SourceTextures.Num();
 	for (int32 ti = 0; ti < NumTextures; ++ti)
 	{
 		const FTextureInfo& SourceTex = SourceMaterialInfo.SourceTextures[ti];
 
-		if (SourceTex.bIsNormalMap)
+		if (SourceTex.Constraint == ETextureBakingConstraint::UseExistingTexture)
+		{
+			FMaterialParameterInfo ParamInfo(SourceTex.ParameterName);
+			Material->SetTextureParameterValueByInfo(ParamInfo, SourceTex.Texture);
+		}
+		else if (SourceTex.bIsNormalMap)
 		{
-			if (ensure(PreviewNormalMap))
+			if (PreviewNormalMap)
 			{
 				FMaterialParameterInfo ParamInfo(SourceTex.ParameterName);
 				Material->SetTextureParameterValueByInfo(ParamInfo, PreviewNormalMap);
@@ -1051,13 +1182,39 @@ void UGenerateStaticMeshLODProcess::WriteDerivedTextures(bool bCreatingNewStatic
 	// for a filename that already exists
 	TMap<UTexture2D*, UTexture2D*> WrittenSourceToDerived;
 
-	// write derived textures
+	// figure out which textures we might write have actually been referenced by a
+	// Material we are going to write. We do not want to write unreferenced textures.
+	bool bAnyMaterialUsedGeneratedNormalMap = false;
+	bool bAnyMaterialUsedMultiTexture = false;
 	int32 NumMaterials = SourceMaterials.Num();
+	for (int32 mi = 0; mi < NumMaterials; ++mi)
+	{
+		const FSourceMaterialInfo& SourceMaterialInfo = SourceMaterials[mi];
+		if (SourceMaterialInfo.bIsReusable
+			|| SourceMaterialInfo.bIsPreviouslyGeneratedMaterial
+			|| SourceMaterialInfo.Constraint == EMaterialBakingConstraint::UseExistingMaterial)
+		{
+			continue;
+		}
+		if (SourceMaterialInfo.bHasNormalMap)
+		{
+			bAnyMaterialUsedGeneratedNormalMap = true;
+		}
+		if (CurrentSettings_Texture.bCombineTextures && MultiTextureParameterName.Contains(mi))
+		{
+			bAnyMaterialUsedMultiTexture = true;
+		}
+	}
+
+
+	// write derived textures
 	check(DerivedMaterials.Num() == NumMaterials);
 	for (int32 mi = 0; mi < NumMaterials; ++mi)
 	{
 		const FSourceMaterialInfo& SourceMaterialInfo = SourceMaterials[mi];
-		if (SourceMaterialInfo.bIsReusable || SourceMaterialInfo.bIsPreviouslyGeneratedMaterial)
+		if (SourceMaterialInfo.bIsReusable 
+			|| SourceMaterialInfo.bIsPreviouslyGeneratedMaterial
+			|| SourceMaterialInfo.Constraint == EMaterialBakingConstraint::UseExistingMaterial )
 		{
 			continue;
 		}
@@ -1080,7 +1237,8 @@ void UGenerateStaticMeshLODProcess::WriteDerivedTextures(bool bCreatingNewStatic
 
 			bool bConvertToSRGB = SourceTex.Texture->SRGB;
 
-			if (DerivedTex.bShouldBakeTexture == false)
+			if (DerivedTex.bShouldBakeTexture == false ||
+				DerivedTex.Constraint == ETextureBakingConstraint::UseExistingTexture)
 			{
 				continue;
 			}
@@ -1110,6 +1268,7 @@ void UGenerateStaticMeshLODProcess::WriteDerivedTextures(bool bCreatingNewStatic
 
 
 	// write derived normal map
+	if (bAnyMaterialUsedGeneratedNormalMap) 
 	{
 		FTexture2DBuilder NormapMapBuilder;
 		NormapMapBuilder.Initialize(FTexture2DBuilder::ETextureType::NormalMap, DerivedNormalMapImage.Image.GetDimensions());
@@ -1129,7 +1288,7 @@ void UGenerateStaticMeshLODProcess::WriteDerivedTextures(bool bCreatingNewStatic
 
 
 	// write multi-texture bake result
-	if (CurrentSettings_Texture.bCombineTextures)
+	if (CurrentSettings_Texture.bCombineTextures && bAnyMaterialUsedMultiTexture)
 	{
 		bool bConvertToSRGB = true;
 		
@@ -1228,7 +1387,9 @@ void UGenerateStaticMeshLODProcess::WriteDerivedMaterials(bool bCreatingNewStati
 	for (int32 mi = 0; mi < NumMaterials; ++mi)
 	{
 		const FSourceMaterialInfo& SourceMaterialInfo = SourceMaterials[mi];
-		if (SourceMaterialInfo.bIsReusable || SourceMaterialInfo.bIsPreviouslyGeneratedMaterial)
+		if (SourceMaterialInfo.bIsReusable 
+			|| SourceMaterialInfo.bIsPreviouslyGeneratedMaterial
+			|| SourceMaterialInfo.Constraint == EMaterialBakingConstraint::UseExistingMaterial )
 		{
 			continue;
 		}
@@ -1323,7 +1484,13 @@ void UGenerateStaticMeshLODProcess::UpdateMaterialTextureParameters(UMaterialIns
 	{
 		FTextureInfo& DerivedTex = DerivedMaterialInfo.DerivedTextures[ti];
 
-		if (DerivedTex.bIsNormalMap)
+		if (DerivedTex.Constraint == ETextureBakingConstraint::UseExistingTexture)
+		{
+			UTexture2D* SourceTexture = SourceMaterials[DerivedMaterialInfo.SourceMaterialIndex].SourceTextures[ti].Texture;
+			FMaterialParameterInfo ParamInfo(DerivedTex.ParameterName);
+			Material->SetTextureParameterValueEditorOnly(ParamInfo, SourceTexture);
+		}
+		else if (DerivedTex.bIsNormalMap)
 		{
 			if (ensure(DerivedNormalMapTex))
 			{
@@ -1383,7 +1550,7 @@ void UGenerateStaticMeshLODProcess::WriteDerivedStaticMeshAsset()
 	{
 		if (!SourceMaterials[mi].bIsPreviouslyGeneratedMaterial)	// Skip previously generated
 		{
-			if (SourceMaterials[mi].bIsReusable)
+			if (SourceMaterials[mi].bIsReusable || SourceMaterials[mi].Constraint == EMaterialBakingConstraint::UseExistingMaterial )
 			{
 				NewMaterials.Add(SourceMaterials[mi].SourceMaterial);
 			}
@@ -1477,9 +1644,10 @@ void UGenerateStaticMeshLODProcess::UpdateSourceStaticMeshAsset(bool bSetNewHDSo
 		{
 			DerivedMatSlotIndexMap[mi] = -2;		// these materials do not appear in DerivedMesh or SourceMesh and should be skipped/discarded (todo: and deleted?)
 		}
-		else  if (SourceMaterials[mi].bIsReusable)	// if we can re-use existing material we just rewrite to existing material slot index
+		else  if (SourceMaterials[mi].bIsReusable
+				  || SourceMaterials[mi].Constraint == EMaterialBakingConstraint::UseExistingMaterial)	
 		{
-			DerivedMatSlotIndexMap[mi] = mi;
+			DerivedMatSlotIndexMap[mi] = mi;		// if we can re-use existing material we just rewrite to existing material slot index
 		}
 		else
 		{
diff --git a/Engine/Plugins/Experimental/MeshLODToolset/Source/MeshLODToolset/Private/MeshLODToolsetModule.cpp b/Engine/Plugins/Experimental/MeshLODToolset/Source/MeshLODToolset/Private/MeshLODToolsetModule.cpp
--- a/Engine/Plugins/Experimental/MeshLODToolset/Source/MeshLODToolset/Private/MeshLODToolsetModule.cpp
+++ b/Engine/Plugins/Experimental/MeshLODToolset/Source/MeshLODToolset/Private/MeshLODToolsetModule.cpp
@@ -1,20 +1,56 @@
 // Copyright Epic Games, Inc. All Rights Reserved.
 
 #include "MeshLODToolsetModule.h"
+#include "Tools/LODGenerationSettingsAsset.h"
+#include "AssetTypeActions_Base.h"
 
 #define LOCTEXT_NAMESPACE "FMeshLODToolsetModule"
 
 DEFINE_LOG_CATEGORY(LogMeshLODToolset);
 
+
+/**
+ * Asset Type Actions for UStaticMeshLODGenerationSettings Assets
+ */
+class FAssetTypeActions_StaticMeshLODGenerationSettings : public FAssetTypeActions_Base
+{
+public:
+	virtual FText GetName() const override { return NSLOCTEXT("AssetTypeActions", "FAssetTypeActions_StaticMeshLODGenerationSettings", "AutoLOD Settings"); }
+	virtual FColor GetTypeColor() const override { return FColor( 175, 0, 128 ); }
+	virtual UClass* GetSupportedClass() const override { return UStaticMeshLODGenerationSettings::StaticClass(); }
+	virtual uint32 GetCategories() override { return EAssetTypeCategories::Misc; }
+	virtual bool CanLocalize() const override { return false; }
+};
+
+
+
+
 void FMeshLODToolsetModule::StartupModule()
 {
 	// This code will execute after your module is loaded into memory; the exact timing is specified in the .uplugin file per-module
+
+	// Register asset actions
+	IAssetTools& AssetTools = FModuleManager::LoadModuleChecked<FAssetToolsModule>("AssetTools").Get();
+	{
+		TSharedRef<IAssetTypeActions> StaticMeshLODGenerationSettings_AssetActions = MakeShareable( new FAssetTypeActions_StaticMeshLODGenerationSettings );
+		AssetTools.RegisterAssetTypeActions( StaticMeshLODGenerationSettings_AssetActions );
+		RegisteredAssetTypeActions.Add(StaticMeshLODGenerationSettings_AssetActions);
+	}
 }
 
 void FMeshLODToolsetModule::ShutdownModule()
 {
 	// This function may be called during shutdown to clean up your module.  For modules that support dynamic reloading,
 	// we call this function before unloading the module.
+
+	// Unregister asset actions
+	IAssetTools& AssetTools = FModuleManager::LoadModuleChecked<FAssetToolsModule>("AssetTools").Get();
+	{
+		for (TSharedRef<IAssetTypeActions> RegisteredAssetTypeAction : RegisteredAssetTypeActions)
+		{
+			AssetTools.UnregisterAssetTypeActions(RegisteredAssetTypeAction);
+		}
+	}
 }
 
 #undef LOCTEXT_NAMESPACE
diff --git a/Engine/Plugins/Experimental/MeshLODToolset/Source/MeshLODToolset/Private/Tools/GenerateStaticMeshLODAssetTool.cpp b/Engine/Plugins/Experimental/MeshLODToolset/Source/MeshLODToolset/Private/Tools/GenerateStaticMeshLODAssetTool.cpp
--- a/Engine/Plugins/Experimental/MeshLODToolset/Source/MeshLODToolset/Private/Tools/GenerateStaticMeshLODAssetTool.cpp
+++ b/Engine/Plugins/Experimental/MeshLODToolset/Source/MeshLODToolset/Private/Tools/GenerateStaticMeshLODAssetTool.cpp
@@ -48,6 +48,9 @@
 #include "TargetInterfaces/PrimitiveComponentBackedTarget.h"
 #include "TargetInterfaces/StaticMeshBackedTarget.h"
 #include "ToolTargetManager.h"
+#include "ModelingToolTargetUtil.h"
+
+#include "Tools/LODGenerationSettingsAsset.h"
 
 static_assert(WITH_EDITOR, "Tool being compiled without editor");
 #include "Misc/ScopedSlowTask.h"
@@ -76,6 +79,8 @@ namespace GenerateStaticMeshLODAssetLocals
 		FGenerateStaticMeshLODProcess_TextureSettings GeneratorSettings_Texture;
 		FGenerateStaticMeshLODProcess_UVSettings GeneratorSettings_UV;
 		FGenerateStaticMeshLODProcess_CollisionSettings GeneratorSettings_Collision;
+		TArray<TPair<UMaterialInterface*, UGenerateStaticMeshLODProcess::EMaterialBakingConstraint>> MaterialConstraints;
+		TArray<TPair<UTexture2D*, UGenerateStaticMeshLODProcess::ETextureBakingConstraint>> TextureConstraints;
 		
 		// Outputs
 		// Inherited: 	TUniquePtr<FDynamicMesh3> ResultMesh;
@@ -103,6 +108,14 @@ namespace GenerateStaticMeshLODAssetLocals
 				GenerateProcess->UpdateTextureSettings(GeneratorSettings_Texture);
 				GenerateProcess->UpdateUVSettings(GeneratorSettings_UV);
 				GenerateProcess->UpdateCollisionSettings(GeneratorSettings_Collision);
+				for (TPair<UMaterialInterface*, UGenerateStaticMeshLODProcess::EMaterialBakingConstraint> MatConstraint : MaterialConstraints)
+				{
+					GenerateProcess->UpdateSourceBakeMaterialConstraint(MatConstraint.Key, MatConstraint.Value);
+				}
+				for (TPair<UTexture2D*, UGenerateStaticMeshLODProcess::ETextureBakingConstraint> TexConstraint : TextureConstraints)
+				{
+					GenerateProcess->UpdateSourceBakeTextureConstraint(TexConstraint.Key, TexConstraint.Value);
+				}
 
 				if (Progress && Progress->Cancelled())
 				{
@@ -156,6 +169,20 @@ namespace GenerateStaticMeshLODAssetLocals
 			Op->GeneratorSettings_UV = AutoLODTool->BasicProperties->UVGeneration;
 			Op->GeneratorSettings_Collision.CollisionGroupLayerName = AutoLODTool->BasicProperties->CollisionGroupLayerName;
 			Op->GeneratorSettings_Collision = AutoLODTool->BasicProperties->SimpleCollision;
+
+			Op->MaterialConstraints.Reset();
+			for (const FGenerateStaticMeshLOD_MaterialConfig& MatConfig : AutoLODTool->TextureProperties->Materials)
+			{
+				Op->MaterialConstraints.Add( TPair<UMaterialInterface*, UGenerateStaticMeshLODProcess::EMaterialBakingConstraint>( 
+					MatConfig.Material, static_cast<UGenerateStaticMeshLODProcess::EMaterialBakingConstraint>(MatConfig.Constraint) ) );
+			}
+			Op->TextureConstraints.Reset();
+			for (const FGenerateStaticMeshLOD_TextureConfig& TexConfig : AutoLODTool->TextureProperties->Textures)
+			{
+				Op->TextureConstraints.Add( TPair<UTexture2D*, UGenerateStaticMeshLODProcess::ETextureBakingConstraint>( 
+					TexConfig.Texture, static_cast<UGenerateStaticMeshLODProcess::ETextureBakingConstraint>(TexConfig.Constraint) ) );
+			}
+
 			Op->SetResultTransform(ResultTransform);
 			return Op;
 		}
@@ -201,6 +228,15 @@ UInteractiveTool* UGenerateStaticMeshLODAssetToolBuilder::BuildTool(const FToolB
 
 
 
+void UGenerateStaticMeshLODAssetToolPresetProperties::PostAction(EGenerateLODAssetToolPresetAction Action)
+{
+	if (ParentTool.IsValid())
+	{
+		ParentTool->RequestPresetAction(Action);
+	}
+}
+
+
 /*
  * Tool
  */
@@ -225,12 +261,21 @@ void UGenerateStaticMeshLODAssetTool::Setup()
 
 	GenerateProcess = NewObject<UGenerateStaticMeshLODProcess>(this);
 
-	IPrimitiveComponentBackedTarget* SourceComponent = TargetComponentInterface(0);
-	UStaticMeshComponent* StaticMeshComponent = CastChecked<UStaticMeshComponent>(SourceComponent->GetOwnerComponent());
+	UStaticMeshComponent* StaticMeshComponent = CastChecked<UStaticMeshComponent>(UE::ToolTarget::GetTargetComponent(Targets[0]));
 	UStaticMesh* StaticMesh = StaticMeshComponent->GetStaticMesh();
 
+	FComponentMaterialSet ComponentMaterials = UE::ToolTarget::GetMaterialSet(Targets[0], false);
+	FComponentMaterialSet AssetMaterials = UE::ToolTarget::GetMaterialSet(Targets[0], true);
+	if (ComponentMaterials != AssetMaterials)
+	{
+		GetToolManager()->DisplayMessage(
+			LOCTEXT("GenerateStaticMeshLODAssetTool_DifferentMaterials", "Selected Component and StaticMesh Asset have different Material Sets. Asset Materials will be used."),
+			EToolMessageLevel::UserWarning);
+	}
+	
+
 	FProgressCancel Progress;
-	FScopedSlowTask SlowTask(2, LOCTEXT("UGenerateStaticMeshLODAssetTool_Setup", "Initializing tool ..."));
+	FScopedSlowTask SlowTask(2, LOCTEXT("UGenerateStaticMeshLODAssetTool_Setup", "Initializing AutoLOD Generator ..."));
 	SlowTask.MakeDialog();
 
 	if (StaticMesh)
@@ -259,6 +304,12 @@ void UGenerateStaticMeshLODAssetTool::Setup()
 	OutputProperties->GeneratedSuffix = TEXT("_AutoLOD");
 	OutputProperties->RestoreProperties(this);
 
+	PresetProperties = NewObject<UGenerateStaticMeshLODAssetToolPresetProperties>(this);
+	PresetProperties->RestoreProperties(this);
+	PresetProperties->Initialize(this);
+	AddToolPropertySource(PresetProperties);
+	PresetProperties->WatchProperty(PresetProperties->Preset, [this](TWeakObjectPtr<UStaticMeshLODGenerationSettings>) { OnPresetSelectionChanged(); });
+
 	SlowTask.EnterProgressFrame(1);
 	BasicProperties = NewObject<UGenerateStaticMeshLODAssetToolProperties>(this);
 	AddToolPropertySource(BasicProperties);
@@ -278,11 +329,25 @@ void UGenerateStaticMeshLODAssetTool::Setup()
 	BasicProperties->WatchProperty(BasicProperties->CollisionGroupLayerName, [this](FName) { OnSettingsModified(); });
 	BasicProperties->InitializeGroupLayers(&(GenerateProcess->GetSourceMesh()));
 
+	TextureProperties = NewObject<UGenerateStaticMeshLODAssetToolTextureProperties>(this);
+	AddToolPropertySource(TextureProperties);
+	for (UTexture2D* BakeTexture : GenerateProcess->GetSourceBakeTextures())
+	{
+		TextureProperties->Textures.Add( FGenerateStaticMeshLOD_TextureConfig{ BakeTexture, EGenerateStaticMeshLOD_BakeConstraint::NoConstraint } );
+	}
+	TextureProperties->WatchProperty(TextureProperties->Textures, [this](TArray<FGenerateStaticMeshLOD_TextureConfig> NewValues) { OnSettingsModified(); });
+	for (UMaterialInterface* BakeMaterial : GenerateProcess->GetSourceBakeMaterials())
+	{
+		TextureProperties->Materials.Add( FGenerateStaticMeshLOD_MaterialConfig{ BakeMaterial, EGenerateStaticMeshLOD_BakeConstraint::NoConstraint } );
+	}
+	TextureProperties->WatchProperty(TextureProperties->Materials, [this](TArray<FGenerateStaticMeshLOD_MaterialConfig> NewValues) { OnSettingsModified(); });
+
+
 	FBoxSphereBounds Bounds = StaticMeshComponent->Bounds;
-	FTransform PreviewTransform = SourceComponent->GetWorldTransform();
-	PreviewTransform.AddToTranslation(FVector(0, 2.5f*Bounds.BoxExtent.Y, 0));
+	UE::Geometry::FTransform3d PreviewTransform = UE::ToolTarget::GetLocalToWorldTransform(Targets[0]);
+	PreviewTransform.SetTranslation(PreviewTransform.GetTranslation() + 2.5 * (double)Bounds.BoxExtent.Y * FVector3d::UnitY());
 
-	this->OpFactory = MakeUnique<FGenerateStaticMeshLODAssetOperatorFactory>(this, (UE::Geometry::FTransform3d)PreviewTransform);
+	this->OpFactory = MakeUnique<FGenerateStaticMeshLODAssetOperatorFactory>(this, PreviewTransform);
 	PreviewWithBackgroundCompute = NewObject<UMeshOpPreviewWithBackgroundCompute>(this, "Preview");
 	PreviewWithBackgroundCompute->Setup(this->TargetWorld, this->OpFactory.Get());
 	PreviewWithBackgroundCompute->PreviewMesh->SetTangentsMode(EDynamicMeshComponentTangentsMode::ExternallyProvided);
@@ -321,7 +386,7 @@ void UGenerateStaticMeshLODAssetTool::Setup()
 			PreviewTextures = PreviewMaterialSet.Textures;
 			PreviewMaterials = PreviewMaterialSet.Materials;
 			PreviewWithBackgroundCompute->PreviewMesh->SetMaterials(PreviewMaterials);
-			BasicProperties->PreviewTextures = PreviewTextures;
+			TextureProperties->PreviewTextures = PreviewTextures;
 		}
 
 		GenerateProcess->GraphEvalCriticalSection.Unlock();
@@ -340,7 +405,10 @@ void UGenerateStaticMeshLODAssetTool::Setup()
 	CollisionVizSettings->WatchProperty(CollisionVizSettings->bShowHidden, [this](bool bNewValue) { bCollisionVisualizationDirty = true; });
 
 	CollisionPreview = NewObject<UPreviewGeometry>(this);
-	CollisionPreview->CreateInWorld(TargetWorld, PreviewTransform);
+	CollisionPreview->CreateInWorld(TargetWorld, (FTransform)PreviewTransform);
+
+	// read Preset if we started Tool with one already selected
+	OnPresetSelectionChanged();		
 
 	// Pop up notifications for any warnings
 	for ( const FProgressCancel::FMessageInfo& Warning : Progress.Warnings )
@@ -358,11 +426,14 @@ void UGenerateStaticMeshLODAssetTool::OnSettingsModified()
 	PreviewWithBackgroundCompute->InvalidateResult();
 }
 
+
+
 void UGenerateStaticMeshLODAssetTool::Shutdown(EToolShutdownType ShutdownType)
 {
 	OutputProperties->SaveProperties(this);
 	BasicProperties->SaveProperties(this);
 	CollisionVizSettings->SaveProperties(this);
+	PresetProperties->SaveProperties(this);
 
 	CollisionPreview->Disconnect();
 	CollisionPreview = nullptr;
@@ -423,7 +494,7 @@ void UGenerateStaticMeshLODAssetTool::UpdateCollisionVisualization()
 void UGenerateStaticMeshLODAssetTool::CreateNewAsset()
 {
 	check(PreviewWithBackgroundCompute->HaveValidResult());
-	GenerateProcess->CalculateDerivedPathName(OutputProperties->NewAssetName, OutputProperties->GeneratedSuffix);
+	GenerateProcess->UpdateDerivedPathName(OutputProperties->NewAssetName, OutputProperties->GeneratedSuffix);
 
 	check(GenerateProcess->GraphEvalCriticalSection.TryLock());		// No ops should be running
 	GenerateProcess->WriteDerivedAssetData();
@@ -435,7 +506,7 @@ void UGenerateStaticMeshLODAssetTool::CreateNewAsset()
 void UGenerateStaticMeshLODAssetTool::UpdateExistingAsset()
 {
 	check(PreviewWithBackgroundCompute->HaveValidResult());
-	GenerateProcess->CalculateDerivedPathName(OutputProperties->NewAssetName, OutputProperties->GeneratedSuffix);
+	GenerateProcess->UpdateDerivedPathName(OutputProperties->NewAssetName, OutputProperties->GeneratedSuffix);
 
 	check(GenerateProcess->GraphEvalCriticalSection.TryLock());		// No ops should be running
 
@@ -451,6 +522,56 @@ void UGenerateStaticMeshLODAssetTool::UpdateExistingAsset()
 
 
 
+void UGenerateStaticMeshLODAssetTool::RequestPresetAction(EGenerateLODAssetToolPresetAction ActionType)
+{
+	UStaticMeshLODGenerationSettings* CurrentPreset = PresetProperties->Preset.Get();
+	if (CurrentPreset == nullptr)
+	{
+		GetToolManager()->DisplayMessage(
+			LOCTEXT("GenerateStaticMeshLODAssetTool_NoPresetSelected", "No Preset Asset is currently set in the Preset Settings"), EToolMessageLevel::UserError);
+	}
+
+	if (ActionType == EGenerateLODAssetToolPresetAction::ReadFromPreset)
+	{
+		OnPresetSelectionChanged();
+	}
+	else if (ActionType == EGenerateLODAssetToolPresetAction::WriteToPreset)
+	{
+		GetToolManager()->BeginUndoTransaction(LOCTEXT("WriteToPresetAction", "Write Preset"));
+
+		CurrentPreset->SetFlags(RF_Transactional);
+		CurrentPreset->Modify();
+
+		CurrentPreset->Preprocessing = BasicProperties->Preprocessing;
+		CurrentPreset->MeshGeneration = BasicProperties->MeshGeneration;
+		CurrentPreset->Simplification = BasicProperties->Simplification;
+		CurrentPreset->TextureBaking = BasicProperties->TextureBaking;
+		CurrentPreset->UVGeneration = BasicProperties->UVGeneration;
+		CurrentPreset->SimpleCollision = BasicProperties->SimpleCollision;
+
+		CurrentPreset->PostEditChange();
+
+		GetToolManager()->EndUndoTransaction();
+	}
+}
+
+
+void UGenerateStaticMeshLODAssetTool::OnPresetSelectionChanged()
+{
+	UStaticMeshLODGenerationSettings* ApplyPreset = PresetProperties->Preset.Get();
+	if (ApplyPreset)
+	{
+		BasicProperties->Preprocessing = ApplyPreset->Preprocessing;
+		BasicProperties->MeshGeneration = ApplyPreset->MeshGeneration;
+		BasicProperties->Simplification = ApplyPreset->Simplification;
+		BasicProperties->TextureBaking = ApplyPreset->TextureBaking;
+		BasicProperties->UVGeneration = ApplyPreset->UVGeneration;
+		BasicProperties->SimpleCollision = ApplyPreset->SimpleCollision;
+		OnSettingsModified();
+	}
+}
+
+
 
 
 
diff --git a/Engine/Plugins/Experimental/MeshLODToolset/Source/MeshLODToolset/Private/Tools/LODGenerationSettingsAsset.cpp b/Engine/Plugins/Experimental/MeshLODToolset/Source/MeshLODToolset/Private/Tools/LODGenerationSettingsAsset.cpp
new file mode 100644
--- /dev/null
+++ b/Engine/Plugins/Experimental/MeshLODToolset/Source/MeshLODToolset/Private/Tools/LODGenerationSettingsAsset.cpp
@@ -0,0 +1,4 @@
+// Copyright Epic Games, Inc. All Rights Reserved.
+
+#include "Tools/LODGenerationSettingsAsset.h"
+
diff --git a/Engine/Plugins/Experimental/MeshLODToolset/Source/MeshLODToolset/Private/Tools/LODGenerationSettingsFactory.cpp b/Engine/Plugins/Experimental/MeshLODToolset/Source/MeshLODToolset/Private/Tools/LODGenerationSettingsFactory.cpp
new file mode 100644
--- /dev/null
+++ b/Engine/Plugins/Experimental/MeshLODToolset/Source/MeshLODToolset/Private/Tools/LODGenerationSettingsFactory.cpp
@@ -0,0 +1,42 @@
+// Copyright Epic Games, Inc. All Rights Reserved.
+
+#include "Tools/LODGenerationSettingsFactory.h"
+#include "Tools/LODGenerationSettingsAsset.h"
+#include "AssetTypeCategories.h"
+
+#define LOCTEXT_NAMESPACE "UStaticMeshLODGenerationSettingsFactory"
+
+UStaticMeshLODGenerationSettingsFactory::UStaticMeshLODGenerationSettingsFactory(const FObjectInitializer& ObjectInitializer)
+	: Super(ObjectInitializer)
+{
+	SupportedClass = UStaticMeshLODGenerationSettings::StaticClass();
+	bCreateNew = true;
+	bText = true;
+	bEditAfterNew = true;
+	bEditorImport = false;
+}
+
+UObject* UStaticMeshLODGenerationSettingsFactory::FactoryCreateNew(UClass* Class, UObject* InParent, FName Name, EObjectFlags Flags, UObject* Context, FFeedbackContext* Warn)
+{
+	UStaticMeshLODGenerationSettings* NewSettings = NewObject<UStaticMeshLODGenerationSettings>(InParent, Class, Name, Flags);
+	check(NewSettings);
+	return NewSettings;
+}
+
+
+uint32 UStaticMeshLODGenerationSettingsFactory::GetMenuCategories() const
+{
+	return EAssetTypeCategories::None;
+}
+
+FText UStaticMeshLODGenerationSettingsFactory::GetDisplayName() const
+{
+	return LOCTEXT("MenuEntry", "AutoLOD Settings");
+}
+
+FString UStaticMeshLODGenerationSettingsFactory::GetDefaultNewAssetName() const
+{
+	return TEXT("AutoLODSettings");
+}
+
+#undef LOCTEXT_NAMESPACE
\ No newline at end of file
diff --git a/Engine/Plugins/Experimental/MeshLODToolset/Source/MeshLODToolset/Public/Graphs/GenerateStaticMeshLODProcess.h b/Engine/Plugins/Experimental/MeshLODToolset/Source/MeshLODToolset/Public/Graphs/GenerateStaticMeshLODProcess.h
--- a/Engine/Plugins/Experimental/MeshLODToolset/Source/MeshLODToolset/Public/Graphs/GenerateStaticMeshLODProcess.h
+++ b/Engine/Plugins/Experimental/MeshLODToolset/Source/MeshLODToolset/Public/Graphs/GenerateStaticMeshLODProcess.h
@@ -160,14 +160,6 @@ struct FGenerateStaticMeshLODProcess_UVSettings
 
 	UPROPERTY(EditAnywhere, Category = "PatchBuilder", meta = (DisplayName = "UV - PatchBuilder", UIMin = "0.0", UIMax = "90.0", ClampMin = "0.0", ClampMax = "180.0", EditConditionHides, EditCondition = "UVMethod == EGenerateStaticMeshLODProcess_AutoUVMethod::PatchBuilder"))
 	FGenerateStaticMeshLODProcess_UVSettings_PatchBuilder PatchBuilder;
-
-	///** Number of smoothing steps to apply in the ExpMap UV Generation method (Smoothing slightly increases distortion but produces more stable results) */
-	//UPROPERTY(EditAnywhere, Category = "PatchBuilder", AdvancedDisplay, meta = (UIMin = "0", UIMax = "25", ClampMin = "0", ClampMax = "1000", EditConditionHides, EditCondition = "UVMethod == EGenerateStaticMeshLODProcess_AutoUVMethod::PatchBuilder"))
-	//int SmoothingSteps = 5;
-
-	///** Smoothing parameter, larger values result in faster smoothing in each step */
-	//UPROPERTY(EditAnywhere, Category = "PatchBuilder", AdvancedDisplay, meta = (UIMin = "0", UIMax = "1.0", ClampMin = "0", ClampMax = "1.0", EditConditionHides, EditCondition = "UVMethod == EGenerateStaticMeshLODProcess_AutoUVMethod::PatchBuilder"))
-	//float SmoothingAlpha = 0.25f;
 };
 
 
@@ -293,8 +285,27 @@ class UGenerateStaticMeshLODProcess : public UObject
 	GENERATED_BODY()
 public:
 
+	/**
+	 * Initialize the Process, this will:
+	 *  - read the Source Mesh, either LOD0 or the HiRes LOD if available
+	 *  - Extract all Materials and Textures from the Source Asset and determine what needs to be baked
+	 *  - Read all the Source Texture Images that need to be baked
+	 *  - call InitializeGenerator() to set up the compute graph
+	 */
 	bool Initialize(UStaticMesh* SourceMesh, FProgressCancel* Progress = nullptr);
 
+
+	/** @return the source UStaticMesh passed to Initialize() */
+	UStaticMesh* GetSourceStaticMesh() const { return SourceStaticMesh; }
+	/** @return the converted source MeshDescription as a FDynamicMesh3, valid after Initialize() finishes */
+	const FDynamicMesh3& GetSourceMesh() const { return SourceMesh; }
+
+
+
+	//
+	// Settings Objects for various stages of the Generator Process
+	//
+
 	const FGenerateStaticMeshLODProcess_PreprocessSettings& GetCurrentPreprocessSettings() const { return CurrentSettings_Preprocess; }
 	void UpdatePreprocessSettings(const FGenerateStaticMeshLODProcess_PreprocessSettings& NewSettings);
 
@@ -314,44 +325,106 @@ public:
 	void UpdateCollisionSettings(const FGenerateStaticMeshLODProcess_CollisionSettings& NewSettings);
 
 
+	//
+	// Material Queries and Constraints
+	//
+
+	/** @return the set of Materials that were identified on the Source Asset */
+	TArray<UMaterialInterface*> GetSourceBakeMaterials() const;
+
+	enum class EMaterialBakingConstraint
+	{
+		NoConstraint = 0,
+		UseExistingMaterial = 1
+	};
+
+	/** Configure the constraint on the specified Source Material */
+	void UpdateSourceBakeMaterialConstraint(UMaterialInterface* Material, EMaterialBakingConstraint Constraint);
+
+
+	//
+	// Texture Queries and Constraints
+	//
+
+	/** @return the set of Textures that were identified as being parameters to the Source Materials, that are identified as potentially needing to be baked to the Derived Asset UVs */
+	TArray<UTexture2D*> GetSourceBakeTextures() const;
+
+
+	enum class ETextureBakingConstraint
+	{
+		NoConstraint = 0,
+		UseExistingTexture = 1
+	};
+
+	/** Configure the constraint on the specified Source Texture */
+	void UpdateSourceBakeTextureConstraint(UTexture2D* Texture, ETextureBakingConstraint Constraint);
+
+
+	//
+	// Output Path / Naming / etc Configuration
+	//
+
+	/**
+	 * Configure the output Name and Suffix for the Mesh Asset. The various strings below will be
+	 * computed based on this input
+	 */
+	void UpdateDerivedPathName(const FString& NewAssetBaseName, const FString& NewAssetSuffix);
 
-	UStaticMesh* GetSourceStaticMesh() const { return SourceStaticMesh; }
-	const FDynamicMesh3& GetSourceMesh() const { return SourceMesh; }
 
 	const FString& GetSourceAssetPath() const { return SourceAssetPath; }
 	const FString& GetSourceAssetFolder() const { return SourceAssetFolder; }
 	const FString& GetSourceAssetName() const { return SourceAssetName; }
-
 	static FString GetDefaultDerivedAssetSuffix() { return TEXT("_AutoLOD"); }
 	const FString& GetDerivedAssetName() const { return DerivedAssetName; }
 
-	void CalculateDerivedPathName(const FString& NewAssetBaseName, const FString& NewAssetSuffix);
+
+	/**
+	 * Given the above configuration, run the Generator to compute the derived Meshes and Textures
+	 */
 
 	bool ComputeDerivedSourceData(FProgressCancel* Progress);
+
+	/** The Generated Mesh */
 	const FDynamicMesh3& GetDerivedLOD0Mesh() const { return DerivedLODMesh; }
+	/** Tangents for the Generated Mesh */
 	const UE::Geometry::FMeshTangentsd& GetDerivedLOD0MeshTangents() const { return DerivedLODMeshTangents; }
+	/** Collision geometry for the Generated Mesh */
 	const UE::Geometry::FSimpleShapeSet3d& GetDerivedCollision() const { return DerivedCollision; }
 
+	struct FPreviewMaterials
+	{
+		TArray<UMaterialInterface*> Materials;
+		TArray<UTexture2D*> Textures;
+	};
+	/** Set of Materials and Textures to be applied to the Generated Mesh */
+	void GetDerivedMaterialsPreview(FPreviewMaterials& MaterialSetOut);
+
+
+	//
+	// Final Output Options
+	//
+
 	/**
-	 * Creates new Asset
+	 * Creates new UStaticMesh Asset, Material Instances, and Textures
 	 */
 	virtual bool WriteDerivedAssetData();
 
 
 	/**
-	 * Updates existing SM Asset
+	 * Update the input UStaticMesh Asset (ie the asset passed to Initialize)
+	 * @param bSetNewHDSourceAsset If true, the original source mesh copied to the UStaticMesh's HiRes SourceModel slot
 	 */
 	virtual void UpdateSourceAsset(bool bSetNewHDSourceAsset = false);
 
 
 
-	struct FPreviewMaterials
-	{
-		TArray<UMaterialInterface*> Materials;
-		TArray<UTexture2D*> Textures;
-	};
-	void GetDerivedMaterialsPreview(FPreviewMaterials& MaterialSetOut);
-
+public:
+	/**
+	 * This FCriticalSection is used to co-ordinate evaluation of the internal compute graph
+	 * by external classes. Generally while doing things like changing settings, reading output, etc,
+	 * the internal compute graph cannot be executing. Often this needs to be locked across multiple
+	 * function calls, so it is exposed publicly.  (TODO: clean this up...)
+	 */
 	FCriticalSection GraphEvalCriticalSection;
 
 protected:
@@ -380,7 +453,8 @@ protected:
 		bool bIsNormalMap = false;
 		bool bIsDefaultTexture = false;
 		bool bShouldBakeTexture = false;
-		bool bIsUsedInMultiTextureBaking = false;
+		bool bIsUsedInMultiTextureBaking = false; 
+		ETextureBakingConstraint Constraint = ETextureBakingConstraint::NoConstraint;
 	};
 
 	int SelectTextureToBake(const TArray<FTextureInfo>& TextureInfos) const;
@@ -396,6 +470,8 @@ protected:
 		bool bIsReusable = false;						// if true, Material doesn't need any texture baking and can be re-used by LOD0
 		bool bIsPreviouslyGeneratedMaterial = false;	// if true, this Material was previously generated by AutoLOD and should be discarded. 
 														// Currently inferred from material being in LOD0 but not HiRes Source
+
+		EMaterialBakingConstraint Constraint = EMaterialBakingConstraint::NoConstraint;
 	};
 
 	// list of initial Source Materials, length equivalent to StaticMesh.StaticMaterials
diff --git a/Engine/Plugins/Experimental/MeshLODToolset/Source/MeshLODToolset/Public/MeshLODToolsetModule.h b/Engine/Plugins/Experimental/MeshLODToolset/Source/MeshLODToolset/Public/MeshLODToolsetModule.h
--- a/Engine/Plugins/Experimental/MeshLODToolset/Source/MeshLODToolset/Public/MeshLODToolsetModule.h
+++ b/Engine/Plugins/Experimental/MeshLODToolset/Source/MeshLODToolset/Public/MeshLODToolsetModule.h
@@ -4,9 +4,12 @@
 
 #include "CoreMinimal.h"
 #include "Modules/ModuleManager.h"
+#include "Templates/SharedPointer.h"
 
 DECLARE_LOG_CATEGORY_EXTERN(LogMeshLODToolset, Log, All);
 
+class IAssetTypeActions;
+
 class FMeshLODToolsetModule : public IModuleInterface
 {
 public:
@@ -14,4 +17,8 @@ public:
 	/** IModuleInterface implementation */
 	virtual void StartupModule() override;
 	virtual void ShutdownModule() override;
+
+
+	// registered asset actions
+	TArray<TSharedRef<IAssetTypeActions>> RegisteredAssetTypeActions;
 };
diff --git a/Engine/Plugins/Experimental/MeshLODToolset/Source/MeshLODToolset/Public/Tools/GenerateStaticMeshLODAssetTool.h b/Engine/Plugins/Experimental/MeshLODToolset/Source/MeshLODToolset/Public/Tools/GenerateStaticMeshLODAssetTool.h
--- a/Engine/Plugins/Experimental/MeshLODToolset/Source/MeshLODToolset/Public/Tools/GenerateStaticMeshLODAssetTool.h
+++ b/Engine/Plugins/Experimental/MeshLODToolset/Source/MeshLODToolset/Public/Tools/GenerateStaticMeshLODAssetTool.h
@@ -20,6 +20,7 @@ struct FMeshDescription;
 class UDynamicMeshComponent;
 class UMeshOpPreviewWithBackgroundCompute;
 class UGenerateStaticMeshLODAssetTool;
+class UStaticMeshLODGenerationSettings;
 namespace GenerateStaticMeshLODAssetLocals
 {
 	class FGenerateStaticMeshLODAssetOperatorFactory;
@@ -28,8 +29,8 @@ namespace GenerateStaticMeshLODAssetLocals
 UENUM()
 enum class EGenerateLODAssetOutputMode : uint8
 {
-	UpdateExistingAsset = 0,
-	CreateNewAsset = 1
+	CreateNewAsset = 0,
+	UpdateExistingAsset = 1
 };
 
 
@@ -57,23 +58,65 @@ class MESHLODTOOLSET_API UGenerateStaticMeshLODAssetToolOutputProperties : publi
 	GENERATED_BODY()
 public:
 	/** Whether to modify the static mesh in place or create a new one. */
-	UPROPERTY(EditAnywhere, Category = AssetOptions, meta = (TransientToolProperty))
-	EGenerateLODAssetOutputMode OutputMode = EGenerateLODAssetOutputMode::UpdateExistingAsset;
+	UPROPERTY(EditAnywhere, Category = "Output Options", meta = (TransientToolProperty))
+	EGenerateLODAssetOutputMode OutputMode = EGenerateLODAssetOutputMode::CreateNewAsset;
 
 	/** Base name for newly-generated asset */
-	UPROPERTY(EditAnywhere, Category = AssetOptions, meta = (TransientToolProperty, EditCondition = "OutputMode == EGenerateLODAssetOutputMode::CreateNewAsset"))
+	UPROPERTY(EditAnywhere, Category = "Output Options", meta = (TransientToolProperty, EditConditionHides, EditCondition = "OutputMode == EGenerateLODAssetOutputMode::CreateNewAsset"))
 	FString NewAssetName;
 
-	/** Suffix to append to newly-generated assets */
-	UPROPERTY(EditAnywhere, Category = AssetOptions, meta = (TransientToolProperty))
+	/** If the Asset doesn't already have a HiRes source, store the input mesh as the HiRes source */
+	UPROPERTY(EditAnywhere, Category = "Output Options", meta = (EditConditionHides, EditCondition = "OutputMode == EGenerateLODAssetOutputMode::UpdateExistingAsset"))
+	bool bSaveInputAsHiResSource = true;
+
+	/** Suffix to append to newly-generated Asset (Meshes, Textures, Materials, etc) */
+	UPROPERTY(EditAnywhere, Category = "Output Options", meta = (TransientToolProperty))
 	FString GeneratedSuffix;
+};
 
-	/** If this is checked and the asset doesn't already have a HiRes source, the input mesh will be stored as the HiRes source */
-	UPROPERTY(EditAnywhere, Category = AssetOptions, meta = (EditCondition = "OutputMode == EGenerateLODAssetOutputMode::UpdateExistingAsset"))
-	bool bSaveInputAsHiResSource = true;
+
+
+
+UENUM()
+enum class EGenerateLODAssetToolPresetAction : uint8
+{
+	ReadFromPreset,
+	WriteToPreset
 };
 
 
+UCLASS()
+class MESHLODTOOLSET_API UGenerateStaticMeshLODAssetToolPresetProperties : public UInteractiveToolPropertySet
+{
+	GENERATED_BODY()
+
+public:
+	TWeakObjectPtr<UGenerateStaticMeshLODAssetTool> ParentTool;
+	void Initialize(UGenerateStaticMeshLODAssetTool* ParentToolIn) { ParentTool = ParentToolIn; }
+	virtual void PostAction(EGenerateLODAssetToolPresetAction Action);
+
+public:
+	/** Preset Asset represents a set of Saved Settings for this Tool */
+	UPROPERTY(EditAnywhere, Category = Preset, meta = (DisplayName = "Settings Preset"))
+	TWeakObjectPtr<UStaticMeshLODGenerationSettings> Preset;
+
+	/** Save the current Tool settings to the Preset Asset */
+	UFUNCTION(CallInEditor, Category = Preset, meta = (DisplayPriority = 0))
+	void ReadFromPreset()
+	{
+		PostAction(EGenerateLODAssetToolPresetAction::ReadFromPreset);
+	}
+
+	/** Read the current Tool settings from the Preset Asset */
+	UFUNCTION(CallInEditor, Category = Preset, meta = (DisplayPriority = 1))
+	void WriteToPreset()
+	{
+		PostAction(EGenerateLODAssetToolPresetAction::WriteToPreset);
+	}
+};
+
+
+
 /**
  * Standard properties
  */
@@ -83,22 +126,22 @@ class MESHLODTOOLSET_API UGenerateStaticMeshLODAssetToolProperties : public UInt
 	GENERATED_BODY()
 
 public:
-	UPROPERTY(EditAnywhere, Category = Settings)
+	UPROPERTY(EditAnywhere, Category = "Generator Configuration")
 	FGenerateStaticMeshLODProcess_PreprocessSettings Preprocessing;
 
-	UPROPERTY(EditAnywhere, Category = Settings)
+	UPROPERTY(EditAnywhere, Category = "Generator Configuration", meta = (ExpandByDefault))
 	FGenerateStaticMeshLODProcessSettings MeshGeneration;
 
-	UPROPERTY(EditAnywhere, Category = Settings)
+	UPROPERTY(EditAnywhere, Category = "Generator Configuration", meta = (ExpandByDefault))
 	FGenerateStaticMeshLODProcess_SimplifySettings Simplification;
 
-	UPROPERTY(EditAnywhere, Category = Settings)
+	UPROPERTY(EditAnywhere, Category = "Generator Configuration", meta = (ExpandByDefault))
 	FGenerateStaticMeshLODProcess_TextureSettings TextureBaking;
 
-	UPROPERTY(EditAnywhere, Category = Settings)
+	UPROPERTY(EditAnywhere, Category = "Generator Configuration")
 	FGenerateStaticMeshLODProcess_UVSettings UVGeneration;
 
-	UPROPERTY(EditAnywhere, Category = Settings)
+	UPROPERTY(EditAnywhere, Category = "Generator Configuration")
 	FGenerateStaticMeshLODProcess_CollisionSettings SimpleCollision;
 
 
@@ -139,14 +182,67 @@ public:
 		}
 	}
 
-	// ------------
+};
 
-	UPROPERTY(VisibleAnywhere, Category = Previews)
-	TArray<UTexture2D*> PreviewTextures;
 
+UENUM()
+enum class EGenerateStaticMeshLOD_BakeConstraint
+{
+	NoConstraint = 0,
+	DoNotBake = 1
+};
+
+USTRUCT()
+struct MESHLODTOOLSET_API FGenerateStaticMeshLOD_TextureConfig
+{
+	GENERATED_BODY()
+
+	UPROPERTY(VisibleAnywhere, Category = Texture)
+	UTexture2D* Texture = nullptr;
+
+	UPROPERTY(EditAnywhere, Category = Texture)
+	EGenerateStaticMeshLOD_BakeConstraint Constraint = EGenerateStaticMeshLOD_BakeConstraint::NoConstraint;
+
+	bool operator==(const FGenerateStaticMeshLOD_TextureConfig& Other) const
+	{
+		return Texture == Other.Texture && Constraint == Other.Constraint;
+	}
 };
 
 
+USTRUCT()
+struct MESHLODTOOLSET_API FGenerateStaticMeshLOD_MaterialConfig
+{
+	GENERATED_BODY()
+
+	UPROPERTY(VisibleAnywhere, Category = Material)
+	UMaterialInterface* Material = nullptr;
+
+	UPROPERTY(EditAnywhere, Category = Texture)
+	EGenerateStaticMeshLOD_BakeConstraint Constraint = EGenerateStaticMeshLOD_BakeConstraint::NoConstraint;
+
+	bool operator==(const FGenerateStaticMeshLOD_MaterialConfig& Other) const
+	{
+		return Material == Other.Material && Constraint == Other.Constraint;
+	}
+};
+
+
+
+UCLASS()
+class MESHLODTOOLSET_API UGenerateStaticMeshLODAssetToolTextureProperties : public UInteractiveToolPropertySet
+{
+	GENERATED_BODY()
+public:
+	UPROPERTY(EditAnywhere, Category = "Source Materials Configuration", meta = (TransientToolProperty))
+	TArray<FGenerateStaticMeshLOD_MaterialConfig> Materials;
+
+	UPROPERTY(EditAnywhere, Category = "Source Textures Configuration", meta = (TransientToolProperty))
+	TArray<FGenerateStaticMeshLOD_TextureConfig> Textures;
+
+	UPROPERTY(VisibleAnywhere, Category = "Baked Texture Previews")
+	TArray<UTexture2D*> PreviewTextures;
+};
 
 
 
@@ -171,18 +267,33 @@ public:
 	virtual bool HasAccept() const override { return true; }
 	virtual bool CanAccept() const override;
 
+	virtual void RequestPresetAction(EGenerateLODAssetToolPresetAction ActionType);
+
 protected:
 
-	friend class GenerateStaticMeshLODAssetLocals::FGenerateStaticMeshLODAssetOperatorFactory;
+	UPROPERTY()
+	TObjectPtr<UGenerateStaticMeshLODAssetToolOutputProperties> OutputProperties = nullptr;
+
+	UPROPERTY()
+	TObjectPtr<UGenerateStaticMeshLODAssetToolProperties> BasicProperties = nullptr;
+
+	UPROPERTY()
+	TObjectPtr<UGenerateStaticMeshLODAssetToolPresetProperties> PresetProperties = nullptr;
 
 	UPROPERTY()
-	UGenerateStaticMeshLODAssetToolOutputProperties* OutputProperties;
+	TObjectPtr<UGenerateStaticMeshLODAssetToolTextureProperties> TextureProperties = nullptr;
 
 	UPROPERTY()
-	UGenerateStaticMeshLODAssetToolProperties* BasicProperties;
+	TObjectPtr<UCollisionGeometryVisualizationProperties> CollisionVizSettings = nullptr;
+
+
+
+protected:
+
+	friend class GenerateStaticMeshLODAssetLocals::FGenerateStaticMeshLODAssetOperatorFactory;
 
 	UPROPERTY()
-	UMeshOpPreviewWithBackgroundCompute* PreviewWithBackgroundCompute = nullptr;
+	TObjectPtr<UMeshOpPreviewWithBackgroundCompute> PreviewWithBackgroundCompute = nullptr;
 
 	UPROPERTY()
 	TArray<UTexture2D*> PreviewTextures;
@@ -194,22 +305,19 @@ protected:
 protected:
 
 	UPROPERTY()
-	UCollisionGeometryVisualizationProperties* CollisionVizSettings = nullptr;
-
-	UPROPERTY()
-	UPhysicsObjectToolPropertySet* ObjectData = nullptr;
+	TObjectPtr<UPhysicsObjectToolPropertySet> ObjectData = nullptr;
 
 	UPROPERTY()
-	UMaterialInterface* LineMaterial = nullptr;
+	TObjectPtr<UMaterialInterface> LineMaterial = nullptr;
 
 	UPROPERTY()
-	UPreviewGeometry* CollisionPreview;
+	TObjectPtr<UPreviewGeometry> CollisionPreview;
 
 protected:
 	UWorld* TargetWorld;
 
 	UPROPERTY()
-	UGenerateStaticMeshLODProcess* GenerateProcess;
+	TObjectPtr<UGenerateStaticMeshLODProcess> GenerateProcess;
 
 	TUniquePtr<UE::Geometry::IDynamicMeshOperatorFactory> OpFactory;
 
@@ -222,6 +330,7 @@ protected:
 	void CreateNewAsset();
 	void UpdateExistingAsset();
 
+	void OnPresetSelectionChanged();
 
 
 };
diff --git a/Engine/Plugins/Experimental/MeshLODToolset/Source/MeshLODToolset/Public/Tools/LODGenerationSettingsAsset.h b/Engine/Plugins/Experimental/MeshLODToolset/Source/MeshLODToolset/Public/Tools/LODGenerationSettingsAsset.h
new file mode 100644
--- /dev/null
+++ b/Engine/Plugins/Experimental/MeshLODToolset/Source/MeshLODToolset/Public/Tools/LODGenerationSettingsAsset.h
@@ -0,0 +1,48 @@
+// Copyright Epic Games, Inc. All Rights Reserved.
+
+#pragma once
+
+#include "CoreMinimal.h"
+#include "EngineDefines.h"
+#include "UObject/ObjectMacros.h"
+#include "UObject/UObjectGlobals.h"
+#include "UObject/Object.h"
+#include "Graphs/GenerateStaticMeshLODProcess.h"
+
+#include "LODGenerationSettingsAsset.generated.h"
+
+
+/**
+ * UStaticMeshLODGenerationSettings is intended to be a stored version of the settings used
+ * by UGenerateStaticMeshLODProcess (and the associated UGenerateStaticMeshLODAssetTool). 
+ * This UObject is exposed as an Asset type in the Editor via UStaticMeshLODGenerationSettingsFactory.
+ * 
+ * The Tool uses these serialized settings as a 'Preset', ie the user can save a set
+ * of configured settings, or load previously-saved settings. 
+ */
+UCLASS(Blueprintable)
+class MESHLODTOOLSET_API UStaticMeshLODGenerationSettings : public UObject
+{
+	GENERATED_BODY()
+
+#if WITH_EDITORONLY_DATA
+public:
+	UPROPERTY(EditAnywhere, Category=Settings)
+	FGenerateStaticMeshLODProcess_PreprocessSettings Preprocessing;
+
+	UPROPERTY(EditAnywhere, Category=Settings)
+	FGenerateStaticMeshLODProcessSettings MeshGeneration;
+
+	UPROPERTY(EditAnywhere, Category=Settings)
+	FGenerateStaticMeshLODProcess_SimplifySettings Simplification;
+
+	UPROPERTY(EditAnywhere, Category=Settings)
+	FGenerateStaticMeshLODProcess_TextureSettings TextureBaking;
+
+	UPROPERTY(EditAnywhere, Category=Settings)
+	FGenerateStaticMeshLODProcess_UVSettings UVGeneration;
+
+	UPROPERTY(EditAnywhere, Category=Settings)
+	FGenerateStaticMeshLODProcess_CollisionSettings SimpleCollision;
+#endif
+};
\ No newline at end of file
diff --git a/Engine/Plugins/Experimental/MeshLODToolset/Source/MeshLODToolset/Public/Tools/LODGenerationSettingsFactory.h b/Engine/Plugins/Experimental/MeshLODToolset/Source/MeshLODToolset/Public/Tools/LODGenerationSettingsFactory.h
new file mode 100644
--- /dev/null
+++ b/Engine/Plugins/Experimental/MeshLODToolset/Source/MeshLODToolset/Public/Tools/LODGenerationSettingsFactory.h
@@ -0,0 +1,28 @@
+// Copyright Epic Games, Inc. All Rights Reserved.
+
+#pragma once
+
+#include "CoreMinimal.h"
+#include "UObject/ObjectMacros.h"
+#include "Factories/Factory.h"
+#include "LODGenerationSettingsFactory.generated.h"
+
+/**
+ * Asset Factory for UStaticMeshLODGenerationSettings, which is used to save settings for
+ * the AutoLOD Tool/Process as an Asset in the Editor
+ */
+UCLASS(MinimalAPI, hidecategories = Object)
+class UStaticMeshLODGenerationSettingsFactory : public UFactory
+{
+	GENERATED_UCLASS_BODY()
+
+	//~ Begin UFactory Interface
+	virtual UObject* FactoryCreateNew(UClass* Class, UObject* InParent, FName Name, EObjectFlags Flags, UObject* Context, FFeedbackContext* Warn) override;
+	virtual uint32 GetMenuCategories() const override;
+	virtual FText GetDisplayName() const override;
+	virtual FString GetDefaultNewAssetName() const override;
+	virtual bool ShouldShowInNewMenu() const override { return true; }
+	//~ Begin UFactory Interface	
+};
+
+
diff --git a/Engine/Plugins/Experimental/MeshModelingToolsetExp/Source/MeshModelingToolsExp/Public/Physics/CollisionPropertySets.h b/Engine/Plugins/Experimental/MeshModelingToolsetExp/Source/MeshModelingToolsExp/Public/Physics/CollisionPropertySets.h
--- a/Engine/Plugins/Experimental/MeshModelingToolsetExp/Source/MeshModelingToolsExp/Public/Physics/CollisionPropertySets.h
+++ b/Engine/Plugins/Experimental/MeshModelingToolsetExp/Source/MeshModelingToolsExp/Public/Physics/CollisionPropertySets.h
@@ -125,13 +125,13 @@ class MESHMODELINGTOOLSEXP_API UCollisionGeometryVisualizationProperties : publi
 	GENERATED_BODY()
 public:
 
-	UPROPERTY(EditAnywhere, Category = Visualization)
+	UPROPERTY(EditAnywhere, Category = "Collision Visualization")
 	float LineThickness = 3.0f;
 
-	UPROPERTY(EditAnywhere, Category = Visualization)
+	UPROPERTY(EditAnywhere, Category = "Collision Visualization")
 	bool bShowHidden = false;
 
-	UPROPERTY(EditAnywhere, Category = Visualization)
+	UPROPERTY(EditAnywhere, Category = "Collision Visualization")
 	FColor Color = FColor::Red;
 };
 
diff --git a/Engine/Source/Runtime/GeometryCore/Public/Image/ImageBuilder.h b/Engine/Source/Runtime/GeometryCore/Public/Image/ImageBuilder.h
--- a/Engine/Source/Runtime/GeometryCore/Public/Image/ImageBuilder.h
+++ b/Engine/Source/Runtime/GeometryCore/Public/Image/ImageBuilder.h
@@ -217,6 +217,27 @@ public:
 		return BilinearSample<ScalarType>(PixelCoords, InvalidValue);
 	}
 
+	/**
+	 * @return true if all pixels have the same value
+	 */
+	bool IsConstantValue() const
+	{
+		int64 Num = Dimensions.Num();
+		if (Num >= 1)
+		{
+			PixelType InitialValue = Image[0];
+			for (int64 k = 1; k < Num; ++k)
+			{
+				if (Image[k] != InitialValue)
+				{
+					return false;
+				}
+			}
+		}
+		return true;
+	}
+
+
 	/**
 	 * Very basic downsampling technqiue that just averages NxN pixel blocks. Multi-threaded.
 	 * @param SubSteps each NxN pixel block of this size is averaged into 1 pixel in the output image
