MeshLODToolset: resolve shutdown issue with module unloading
#rnx

diff --git a/Engine/Plugins/Experimental/MeshLODToolset/Source/MeshLODToolset/Private/MeshLODToolsetModule.cpp b/Engine/Plugins/Experimental/MeshLODToolset/Source/MeshLODToolset/Private/MeshLODToolsetModule.cpp
--- a/Engine/Plugins/Experimental/MeshLODToolset/Source/MeshLODToolset/Private/MeshLODToolsetModule.cpp
+++ b/Engine/Plugins/Experimental/MeshLODToolset/Source/MeshLODToolset/Private/MeshLODToolsetModule.cpp
@@ -44,11 +44,12 @@ void FMeshLODToolsetModule::ShutdownModule()
 	// we call this function before unloading the module.
 
 	// Unregister asset actions
-	IAssetTools& AssetTools = FModuleManager::LoadModuleChecked<FAssetToolsModule>("AssetTools").Get();
+	FAssetToolsModule* AssetToolsModule = FModuleManager::GetModulePtr<FAssetToolsModule>("AssetTools");
+	if (AssetToolsModule)
 	{
 		for (TSharedRef<IAssetTypeActions> RegisteredAssetTypeAction : RegisteredAssetTypeActions)
 		{
-			AssetTools.UnregisterAssetTypeActions(RegisteredAssetTypeAction);
+			AssetToolsModule->Get().UnregisterAssetTypeActions(RegisteredAssetTypeAction);
 		}
 	}
 }
