AutoLOD: make mesh computation type configurable using new GeometryFlow Switch node, and expose Normals calculation settings in Graph/Process/Tool

GeometryFlow:
Add TSwitchNode node, uses internal integer to select one of N inputs (N is currently defined in template type, so number of inputs is fixed). Add custom util function UpdateSwitchNodeInputIndex() to simplify changing switch integer.
Add FMeshMakeCleanGeometryNode/Settings, strips out undesired attribute sets and fills small holes, to create better input mesh for simplification
Add FSimpleInPlaceProcessMeshBaseNode, just runs some basic geometry operation, makes it easier to make new nodes. Converted FComputeMeshPerVertexOverlayNormalsNode to this base class and added FComputeMeshPerVertexNormalsNode.
Fixed attribute set handling issue in FSimplifyMeshNode

MeshLODToolset:
Added ECoreMeshGeneratorMode enum to FGenerateMeshLODGraph. This setting drives a Switch node that picks between Vox, Vox+Closure, and Simplify-only mesh generation paths.
Added SourceMeshHint to BuildGraph, allows for variable graph construction. Use to skip MaterialID transfer in single-material case.
Change default collision type to aligned boxes (fast and low visual complexity)
Exposed top-level Mesh Generation type, Normals settings in Graph/Process/Tool, minor UI and default changes

#rnx

diff --git a/Engine/Plugins/Experimental/GeometryFlow/Source/GeometryFlowCore/Public/BaseNodes/SwitchNode.h b/Engine/Plugins/Experimental/GeometryFlow/Source/GeometryFlowCore/Public/BaseNodes/SwitchNode.h
new file mode 100644
--- /dev/null
+++ b/Engine/Plugins/Experimental/GeometryFlow/Source/GeometryFlowCore/Public/BaseNodes/SwitchNode.h
@@ -0,0 +1,164 @@
+// Copyright Epic Games, Inc. All Rights Reserved.
+
+#pragma once
+
+#include "GeometryFlowNode.h"
+
+
+namespace UE
+{
+namespace GeometryFlow
+{
+
+
+/**
+ * TSwitchNode selects one of N inputs and provides it as an Output. This allows for a minimal amount of 
+ * branching / control-flow in a fixed GeometryFlow Graph. Only the active Input is evaluated.
+ * 
+ * The number of inputs is defined in the template type, eg TSwitchNode<FDynamicMesh3, 3, (int)EMeshProcessingDataTypes::DynamicMesh> 
+ * is a switch node with 3 possible inputs. Dynamic number of inputs cannot be easily supported without
+ * changes at the graph evaluation level.
+ * 
+ * Use UpdateSwitchValue() to select the active switch Input
+ * 
+ */
+template<typename T, int NumInputs, int StorageTypeIdentifier>
+class TSwitchNode : public FNode
+{
+protected:
+	using DataContainerType = TMovableData<T, StorageTypeIdentifier>;
+
+	int32 SwitchIndex = 0;
+
+public:
+	static const FString InParamValue(int32 InputNum)
+	{
+		return FString::Printf(TEXT("Input%d"), InputNum);
+	}
+
+	static const FString OutParamValue() { return TEXT("Output"); }
+
+public:
+	TSwitchNode()
+	{
+		// add N input nodes
+		for (int32 k = 0; k < NumInputs; ++k)
+		{
+			AddInput(InParamValue(k), MakeUnique<TBasicNodeInput<T, StorageTypeIdentifier>>());
+		}
+		 
+		AddOutput(OutParamValue(), MakeUnique<TBasicNodeOutput<T, StorageTypeIdentifier>>());
+
+		// assume we can steal input and pass to output
+		for (int32 k = 0; k < NumInputs; ++k)
+		{
+			ConfigureInputFlags(InParamValue(k), FNodeInputFlags::Transformable());
+		}
+
+		SwitchIndex = 0;
+	}
+
+	/**
+	 * This function is used to update the Input Index, ie select which Output will be provided. 
+	 * This function always Invalidates the Output if the Index was changed.
+	 */
+	virtual void UpdateSwitchInputIndex(const int32& NewSwitchIndex)
+	{
+		if (NewSwitchIndex != SwitchIndex)
+		{
+			SwitchIndex = NewSwitchIndex;
+
+			// clear output to force recompute...
+			ClearOutput(OutParamValue());
+		}
+	}
+
+	/**
+	 * Determine which Inputs are required for the desired Output. This replaces the default behavior ("all inputs")
+	 * with only the selected Input returned as a Requirement, which in turn allows the Graph Evaluator to skip evaluation 
+	 * of the other inputs. 
+	 */
+	virtual void CollectRequirements(const TArray<FString>& Outputs, TArray<FEvalRequirement>& RequiredInputsOut) override
+	{
+		RequiredInputsOut.Reset();
+
+		int32 FoundOutputs = 0;
+		for (const FNodeOutputInfo& OutputInfo : NodeOutputs)
+		{
+			if (Outputs.Contains(OutputInfo.Name))
+			{
+				FoundOutputs++;
+			}
+		}
+		if (ensure(FoundOutputs != 0) == false)
+		{
+			return;
+		}
+
+		// only the active input is returned as a requirement
+		FString ActiveInput = InParamValue(SwitchIndex);
+		for ( const FNodeInputInfo & InputInfo : NodeInputs )
+		{
+			if (InputInfo.Name == ActiveInput)
+			{
+				FEvalRequirement Requirement(InputInfo.Name, InputInfo.Input->GetInputFlags());
+				RequiredInputsOut.Add(Requirement);
+			}
+		}	
+	}
+
+
+	virtual void Evaluate(
+		const FNamedDataMap& DatasIn,
+		FNamedDataMap& DatasOut,
+		TUniquePtr<FEvaluationInfo>& EvaluationInfo) override
+	{
+		if (ensure(DatasOut.Contains(OutParamValue())))
+		{
+			bool bAllInputsValid = true;
+			bool bRecomputeRequired = ( IsOutputAvailable(OutParamValue()) == false );
+
+			// get the switch index and clamp it
+			int UseSwitchIndex = SwitchIndex;
+			ensure(UseSwitchIndex >= 0 && UseSwitchIndex < NumInputs);
+			UseSwitchIndex = FMath::Clamp(UseSwitchIndex, 0, NumInputs);
+
+			// select the specified switch input argument
+			TSafeSharedPtr<IData> InputArg = FindAndUpdateInputForEvaluate(InParamValue(UseSwitchIndex), DatasIn, bRecomputeRequired, bAllInputsValid);
+			if (bAllInputsValid)
+			{
+				if (bRecomputeRequired)
+				{
+					bool bIsDataMutable = DatasIn.GetDataFlags(InParamValue(UseSwitchIndex)).bIsMutableData;
+
+					T NewData;
+					if (bIsDataMutable)
+					{
+						InputArg->GiveTo<T>(NewData, StorageTypeIdentifier);
+
+					}
+					else
+					{
+						InputArg->GetDataCopy(NewData, StorageTypeIdentifier);
+					}
+
+					// store new result
+					TSafeSharedPtr<DataContainerType> Result = MakeSafeShared<DataContainerType>();
+					Result->MoveData(MoveTemp(NewData));
+					SetOutput(OutParamValue(), Result);
+
+					EvaluationInfo->CountCompute(this);
+				}
+				DatasOut.SetData(OutParamValue(), GetOutput(OutParamValue()));
+			}
+		}
+	}
+};
+
+
+
+
+}	// end namespace GeometryFlow
+}	// end namespace UE
+
+
diff --git a/Engine/Plugins/Experimental/GeometryFlow/Source/GeometryFlowCore/Public/GeometryFlowGraphUtil.h b/Engine/Plugins/Experimental/GeometryFlow/Source/GeometryFlowCore/Public/GeometryFlowGraphUtil.h
--- a/Engine/Plugins/Experimental/GeometryFlow/Source/GeometryFlowCore/Public/GeometryFlowGraphUtil.h
+++ b/Engine/Plugins/Experimental/GeometryFlow/Source/GeometryFlowCore/Public/GeometryFlowGraphUtil.h
@@ -4,6 +4,7 @@
 
 #include "GeometryFlowGraph.h"
 #include "GeometryFlowCoreNodes.h"
+#include "BaseNodes/SwitchNode.h"
 
 namespace UE
 {
@@ -22,6 +23,18 @@ void UpdateSourceNodeValue(FGraph& Graph, FGraph::FHandle NodeHandle, const type
 }
 
 
+/**
+ * Update the integer index that controls which Input to a TSwitchNode will be provided as it's Output.
+ */
+template<typename SwitchNodeType>
+void UpdateSwitchNodeInputIndex(FGraph& Graph, FGraph::FHandle NodeHandle, int32 NewSwitchIndex)
+{
+	Graph.ApplyToNodeOfType<SwitchNodeType>(NodeHandle, [&](SwitchNodeType& Node)
+	{
+		Node.UpdateSwitchInputIndex(NewSwitchIndex);
+	});
+}
+
 
 
 template<typename SettingsType>
diff --git a/Engine/Plugins/Experimental/GeometryFlow/Source/GeometryFlowMeshProcessing/Private/MeshProcessingNodes/MeshMakeCleanGeometryNode.cpp b/Engine/Plugins/Experimental/GeometryFlow/Source/GeometryFlowMeshProcessing/Private/MeshProcessingNodes/MeshMakeCleanGeometryNode.cpp
new file mode 100644
--- /dev/null
+++ b/Engine/Plugins/Experimental/GeometryFlow/Source/GeometryFlowMeshProcessing/Private/MeshProcessingNodes/MeshMakeCleanGeometryNode.cpp
@@ -0,0 +1,141 @@
+// Copyright Epic Games, Inc. All Rights Reserved.
+
+#include "MeshProcessingNodes/MeshMakeCleanGeometryNode.h"
+#include "DynamicMesh/DynamicMesh3.h"
+#include "DynamicMesh/DynamicMeshAttributeSet.h"
+#include "DynamicMesh/MeshNormals.h"
+#include "MeshQueries.h"
+#include "MeshBoundaryLoops.h"
+#include "Operations/SimpleHoleFiller.h"
+
+
+using namespace UE::Geometry;
+using namespace UE::GeometryFlow;
+
+void FMeshMakeCleanGeometryNode::ProcessMesh(
+	const FNamedDataMap& DatasIn,
+	const FMeshMakeCleanGeometrySettings& SettingsIn,
+	const FDynamicMesh3& MeshIn,
+	FDynamicMesh3& MeshOut) 
+{
+	MeshOut = MeshIn;
+	ApplyMakeCleanGeometry(MeshOut, SettingsIn);
+}
+
+
+void FMeshMakeCleanGeometryNode::ProcessMeshInPlace(
+	const FNamedDataMap& DatasIn,
+	const FMeshMakeCleanGeometrySettings& Settings,
+	FDynamicMesh3& MeshInOut) 
+{
+	ApplyMakeCleanGeometry(MeshInOut, Settings);
+}
+
+// make a triangle fan for the hole - not a great area metric...
+static double EstimateHoleFillArea(const FDynamicMesh3& Mesh, const FEdgeLoop& Loop)
+{
+	FVector3d Centroid = FVector3d::Zero();
+	for (int32 vid : Loop.Vertices)
+	{
+		Centroid += Mesh.GetVertex(vid);
+	}
+	Centroid /= (double)Loop.Vertices.Num();
+
+	double AreaSum = 0;
+	for (int32 eid : Loop.Edges)
+	{
+		FIndex2i EdgeV = Mesh.GetEdgeV(eid);
+		AreaSum += VectorUtil::Area(Mesh.GetVertex(EdgeV.A), Mesh.GetVertex(EdgeV.B), Centroid);
+	}
+	return AreaSum;
+}
+
+
+void FMeshMakeCleanGeometryNode::ApplyMakeCleanGeometry(FDynamicMesh3& Mesh, const FMeshMakeCleanGeometrySettings& Settings)
+{
+	if (Mesh.HasAttributes())
+	{
+		if (Settings.bDiscardAllAttributes)
+		{
+			Mesh.DiscardAttributes();
+		}
+		else
+		{
+			if (Settings.bClearUVs)
+			{
+				Mesh.Attributes()->SetNumUVLayers(0);
+			}
+			if (Settings.bClearVertexColors)
+			{
+				Mesh.Attributes()->DisablePrimaryColors();
+			}
+			if (Settings.bClearTangents && Mesh.Attributes()->NumNormalLayers() > 1)
+			{
+				Mesh.Attributes()->DisableTangents();
+			}
+			if (Settings.bClearNormals)
+			{
+				Mesh.Attributes()->SetNumNormalLayers(0);
+			}
+			if (Settings.bClearMaterialIDs)
+			{
+				Mesh.Attributes()->DisableMaterialID();
+			}
+		}
+	}
+
+	if (Settings.FillHolesEdgeCountThresh > 0 || Settings.FillHolesEstimatedAreaFraction > 0)
+	{
+		double MeshArea = 0.0;
+		if (Settings.FillHolesEstimatedAreaFraction > 0)
+		{
+			FVector2d VolArea = TMeshQueries<FDynamicMesh3>::GetVolumeArea(Mesh);
+			MeshArea = VolArea.Y;
+		}
+
+		FMeshBoundaryLoops BoundaryLoops(&Mesh);
+		int32 InitialHoleCount = BoundaryLoops.Loops.Num();
+
+		for (FEdgeLoop& Loop : BoundaryLoops.Loops)
+		{
+			if (Mesh.IsBoundaryEdge(Loop.Edges[0]))
+			{
+				bool bFill = (Loop.Edges.Num() < Settings.FillHolesEdgeCountThresh);
+				if (bFill == false && Settings.FillHolesEstimatedAreaFraction)
+				{
+					double HoleAreaEst = EstimateHoleFillArea(Mesh, Loop);
+					bFill = ((HoleAreaEst / MeshArea) < Settings.FillHolesEstimatedAreaFraction);
+				}
+
+				if (bFill)
+				{
+					FSimpleHoleFiller Filler(&Mesh, Loop);
+					Filler.Fill();
+				}
+			}
+		}
+
+		//UE_LOG(LogGeometry, Warning, TEXT("GeometryFlow::MeshMakeCleanGeometryNode: Reduced from %d holes to %d holes"), BoundaryLoops.Loops.Num(), FMeshBoundaryLoops(&Mesh).Loops.Num());
+	}
+
+	if (Settings.bOutputMeshVertexNormals)
+	{
+		FMeshNormals::QuickComputeVertexNormals(Mesh, false);
+	}
+	else
+	{
+		Mesh.DiscardVertexNormals();
+	}
+
+	if (Settings.bOutputOverlayVertexNormals)
+	{
+		if (Mesh.HasAttributes())
+		{
+			if (Mesh.Attributes()->NumNormalLayers() < 1)
+			{
+				Mesh.Attributes()->SetNumNormalLayers(1);
+			}
+			FMeshNormals::InitializeOverlayToPerVertexNormals(Mesh.Attributes()->PrimaryNormals(), Settings.bOutputMeshVertexNormals);
+		}
+	}
+}
diff --git a/Engine/Plugins/Experimental/GeometryFlow/Source/GeometryFlowMeshProcessing/Public/MeshProcessingNodes/MeshMakeCleanGeometryNode.h b/Engine/Plugins/Experimental/GeometryFlow/Source/GeometryFlowMeshProcessing/Public/MeshProcessingNodes/MeshMakeCleanGeometryNode.h
new file mode 100644
--- /dev/null
+++ b/Engine/Plugins/Experimental/GeometryFlow/Source/GeometryFlowMeshProcessing/Public/MeshProcessingNodes/MeshMakeCleanGeometryNode.h
@@ -0,0 +1,56 @@
+// Copyright Epic Games, Inc. All Rights Reserved.
+
+#pragma once
+
+#include "MeshProcessingNodes/MeshProcessingBaseNodes.h"
+#include "MeshProcessingNodes/MeshProcessingDataTypes.h"
+
+namespace UE
+{
+namespace GeometryFlow
+{
+
+struct GEOMETRYFLOWMESHPROCESSING_API FMeshMakeCleanGeometrySettings
+{
+	DECLARE_GEOMETRYFLOW_DATA_TYPE_IDENTIFIER(EMeshProcessingDataTypes::MakeCleanGeometrySettings);
+
+	int FillHolesEdgeCountThresh = 8;
+	double FillHolesEstimatedAreaFraction = 0.001;
+
+	bool bDiscardAllAttributes = false;
+	bool bClearUVs = true;
+	bool bClearNormals = true;
+	bool bClearTangents = true;
+	bool bClearVertexColors = true;
+	bool bClearMaterialIDs = false;
+
+	bool bOutputMeshVertexNormals = true;
+	bool bOutputOverlayVertexNormals = true;
+};
+GEOMETRYFLOW_DECLARE_SETTINGS_TYPES(FMeshMakeCleanGeometrySettings, MeshMakeCleanGeometry);
+
+
+class GEOMETRYFLOWMESHPROCESSING_API FMeshMakeCleanGeometryNode : public TProcessMeshWithSettingsBaseNode<FMeshMakeCleanGeometrySettings>
+{
+public:
+
+	void ProcessMesh(
+		const FNamedDataMap& DatasIn,
+		const FMeshMakeCleanGeometrySettings& SettingsIn,
+		const FDynamicMesh3& MeshIn,
+		FDynamicMesh3& MeshOut) override;
+
+	void ProcessMeshInPlace(
+		const FNamedDataMap& DatasIn,
+		const FMeshMakeCleanGeometrySettings& Settings,
+		FDynamicMesh3& MeshInOut) override;
+
+protected:
+
+	void ApplyMakeCleanGeometry(FDynamicMesh3& MeshInOut, const FMeshMakeCleanGeometrySettings& Settings);
+	
+};
+
+
+}	// end namespace GeometryFlow
+}	// end namespace UE
diff --git a/Engine/Plugins/Experimental/GeometryFlow/Source/GeometryFlowMeshProcessing/Public/MeshProcessingNodes/MeshNormalsNodes.h b/Engine/Plugins/Experimental/GeometryFlow/Source/GeometryFlowMeshProcessing/Public/MeshProcessingNodes/MeshNormalsNodes.h
--- a/Engine/Plugins/Experimental/GeometryFlow/Source/GeometryFlowMeshProcessing/Public/MeshProcessingNodes/MeshNormalsNodes.h
+++ b/Engine/Plugins/Experimental/GeometryFlow/Source/GeometryFlowMeshProcessing/Public/MeshProcessingNodes/MeshNormalsNodes.h
@@ -80,6 +80,11 @@ public:
 			MeshInOut.EnableAttributes();
 		}
 		FDynamicMeshNormalOverlay* Normals = MeshInOut.Attributes()->PrimaryNormals();
+		if (Normals == nullptr)
+		{
+			MeshInOut.Attributes()->SetNumNormalLayers(1);
+			Normals = MeshInOut.Attributes()->PrimaryNormals();
+		}
 
 		if (Settings.NormalsType == EComputeNormalsType::PerTriangle)
 		{
@@ -120,45 +125,40 @@ public:
 /**
  * Recompute per-vertex normals in Normals Overlay for input mesh. Can apply in-place.
  */
-class FComputeMeshPerVertexOverlayNormalsNode : public FProcessMeshBaseNode
+class FComputeMeshPerVertexOverlayNormalsNode : public FSimpleInPlaceProcessMeshBaseNode
 {
 public:
-	FComputeMeshPerVertexOverlayNormalsNode()
-	{
-		// we can mutate input mesh
-		ConfigureInputFlags(InParamMesh(), FNodeInputFlags::Transformable());
-	}
-
-	virtual void ProcessMesh(
-		const FNamedDataMap& DatasIn,
-		const FDynamicMesh3& MeshIn,
-		FDynamicMesh3& MeshOut) override
-	{
-		MeshOut = MeshIn;
-		if (MeshOut.HasAttributes() == false)
-		{
-			MeshOut.EnableAttributes();
-		}
-		FDynamicMeshNormalOverlay* Normals = MeshOut.Attributes()->PrimaryNormals();
-		FMeshNormals::InitializeOverlayToPerVertexNormals(Normals, false);
-	}
-
-	virtual void ProcessMeshInPlace(
-		const FNamedDataMap& DatasIn,
-		FDynamicMesh3& MeshInOut) override
+	virtual void ApplyNodeToMesh(FDynamicMesh3& MeshInOut)
 	{
 		if (MeshInOut.HasAttributes() == false)
 		{
 			MeshInOut.EnableAttributes();
 		}
 		FDynamicMeshNormalOverlay* Normals = MeshInOut.Attributes()->PrimaryNormals();
+		if (Normals == nullptr)
+		{
+			MeshInOut.Attributes()->SetNumNormalLayers(1);
+			Normals = MeshInOut.Attributes()->PrimaryNormals();
+		}
 		FMeshNormals::InitializeOverlayToPerVertexNormals(Normals, false);
 	}
-
 };
 
 
 
+/**
+ * Recompute per-vertex normals stored directly on the mesh. Can apply in-place.
+ */
+class FComputeMeshPerVertexNormalsNode : public FSimpleInPlaceProcessMeshBaseNode
+{
+public:
+	virtual void ApplyNodeToMesh(FDynamicMesh3& MeshInOut)
+	{
+		FMeshNormals::QuickComputeVertexNormals(MeshInOut, false);
+	}
+};
+
+
 
 
 }	// end namespace GeometryFlow
diff --git a/Engine/Plugins/Experimental/GeometryFlow/Source/GeometryFlowMeshProcessing/Public/MeshProcessingNodes/MeshProcessingBaseNodes.h b/Engine/Plugins/Experimental/GeometryFlow/Source/GeometryFlowMeshProcessing/Public/MeshProcessingNodes/MeshProcessingBaseNodes.h
--- a/Engine/Plugins/Experimental/GeometryFlow/Source/GeometryFlowMeshProcessing/Public/MeshProcessingNodes/MeshProcessingBaseNodes.h
+++ b/Engine/Plugins/Experimental/GeometryFlow/Source/GeometryFlowMeshProcessing/Public/MeshProcessingNodes/MeshProcessingBaseNodes.h
@@ -118,7 +118,36 @@ public:
 
 
 
+/**
+* FSimpleInPlaceProcessMeshBaseNode provides a standard pattern for a simple "In-Place" mesh
+* processing operation, eg something like recomputing normals where it would not make
+* sense to cache the output and there are no configuration settings. The subclass only
+* has to implement the ApplyNodeToMesh() function.
+*/
+class FSimpleInPlaceProcessMeshBaseNode : public FProcessMeshBaseNode
+{
+public:
+	FSimpleInPlaceProcessMeshBaseNode()
+	{
+		// we can mutate input mesh
+		ConfigureInputFlags(InParamMesh(), FNodeInputFlags::Transformable());
+	}
+
+	virtual void ProcessMesh( const FNamedDataMap& DatasIn, const FDynamicMesh3& MeshIn, FDynamicMesh3& MeshOut) override
+	{
+		MeshOut = MeshIn;
+		ApplyNodeToMesh(MeshOut);
+	}
+
+	virtual void ProcessMeshInPlace( const FNamedDataMap& DatasIn, FDynamicMesh3& MeshInOut) override
+	{
+		ApplyNodeToMesh(MeshInOut);
+	}
+
+	// subclasses only have to implement this
+	virtual void ApplyNodeToMesh(FDynamicMesh3& MeshInOut) = 0;
 
+};
 
 
 
diff --git a/Engine/Plugins/Experimental/GeometryFlow/Source/GeometryFlowMeshProcessing/Public/MeshProcessingNodes/MeshProcessingDataTypes.h b/Engine/Plugins/Experimental/GeometryFlow/Source/GeometryFlowMeshProcessing/Public/MeshProcessingNodes/MeshProcessingDataTypes.h
--- a/Engine/Plugins/Experimental/GeometryFlow/Source/GeometryFlowMeshProcessing/Public/MeshProcessingNodes/MeshProcessingDataTypes.h
+++ b/Engine/Plugins/Experimental/GeometryFlow/Source/GeometryFlowMeshProcessing/Public/MeshProcessingNodes/MeshProcessingDataTypes.h
@@ -36,6 +36,7 @@ enum class EMeshProcessingDataTypes
 	NormalsSettings = (int)EDataTypes::BaseMeshProcessingTypes + 120,
 	TangentsSettings = (int)EDataTypes::BaseMeshProcessingTypes + 121,
 
+	MakeCleanGeometrySettings = (int)EDataTypes::BaseMeshProcessingTypes + 130,
 
 	RecalculateUVsSettings = (int)EDataTypes::BaseMeshProcessingTypes + 150,
 	RepackUVsSettings = (int)EDataTypes::BaseMeshProcessingTypes + 151,
diff --git a/Engine/Plugins/Experimental/GeometryFlow/Source/GeometryFlowMeshProcessing/Public/MeshProcessingNodes/MeshSimplifyNode.h b/Engine/Plugins/Experimental/GeometryFlow/Source/GeometryFlowMeshProcessing/Public/MeshProcessingNodes/MeshSimplifyNode.h
--- a/Engine/Plugins/Experimental/GeometryFlow/Source/GeometryFlowMeshProcessing/Public/MeshProcessingNodes/MeshSimplifyNode.h
+++ b/Engine/Plugins/Experimental/GeometryFlow/Source/GeometryFlowMeshProcessing/Public/MeshProcessingNodes/MeshSimplifyNode.h
@@ -98,17 +98,10 @@ public:
 		{
 			Simplifier.SetEdgeFlipTolerance(1.e-5);
 
-			// eliminate any bowties that might have formed on UV seams.
-			if (FDynamicMeshAttributeSet* Attributes = TargetMesh->Attributes())
+			// eliminate any bowties in attribute layers
+			if (TargetMesh->Attributes())
 			{
-				// @todo parallelize over NumUVLayers?
-				int NumUVLayers = Attributes->NumUVLayers();
-				for (int i = 0; i < NumUVLayers; ++i)
-				{
-					Attributes->GetUVLayer(i)->SplitBowties();
-				}
-				Attributes->PrimaryNormals()->SplitBowties();
-
+				TargetMesh->Attributes()->SplitAllBowties();
 			}
 		}
 
diff --git a/Engine/Plugins/Experimental/MeshLODToolset/Source/MeshLODToolset/Private/Graphs/GenerateMeshLODGraph.cpp b/Engine/Plugins/Experimental/MeshLODToolset/Source/MeshLODToolset/Private/Graphs/GenerateMeshLODGraph.cpp
--- a/Engine/Plugins/Experimental/MeshLODToolset/Source/MeshLODToolset/Private/Graphs/GenerateMeshLODGraph.cpp
+++ b/Engine/Plugins/Experimental/MeshLODToolset/Source/MeshLODToolset/Private/Graphs/GenerateMeshLODGraph.cpp
@@ -6,6 +6,7 @@
 #include "GeometryFlowGraph.h"
 #include "GeometryFlowGraphUtil.h"
 #include "BaseNodes/TransferNode.h"
+#include "BaseNodes/SwitchNode.h"
 
 #include "MeshProcessingNodes/MeshThickenNode.h"
 #include "MeshProcessingNodes/MeshSolidifyNode.h"
@@ -35,6 +36,9 @@ using namespace UE::Geometry;
 using namespace UE::GeometryFlow;
 
 
+typedef UE::GeometryFlow::TSwitchNode<FDynamicMesh3, 3, (int)EMeshProcessingDataTypes::DynamicMesh> FMeshGeneratorSwitchNode;
+
+
 void FGenerateMeshLODGraph::SetSourceMesh(const FDynamicMesh3& SourceMeshIn)
 {
 	UpdateSourceNodeValue<FDynamicMeshSourceNode>(*Graph, MeshSourceNode, SourceMeshIn);
@@ -47,6 +51,12 @@ void FGenerateMeshLODGraph::UpdatePreFilterSettings(const FMeshLODGraphPreFilter
 	CurrentPreFilterSettings = PreFilterSettings;
 }
 
+void FGenerateMeshLODGraph::UpdateCoreMeshGeneratorMode(ECoreMeshGeneratorMode NewMode)
+{
+	UpdateSwitchNodeInputIndex<FMeshGeneratorSwitchNode>(*Graph, MeshGeneratorSwitchNode, static_cast<int32>(NewMode) );
+	CurrentCoreMeshGeneratorMode = NewMode;
+}
+
 
 void FGenerateMeshLODGraph::UpdateSolidifySettings(const FMeshSolidifySettings& SolidifySettings)
 {
@@ -60,12 +70,24 @@ void FGenerateMeshLODGraph::UpdateMorphologySettings(const FVoxClosureSettings&
 	CurrentMorphologySettings = MorphologySettings;
 }
 
+void FGenerateMeshLODGraph::UpdateMeshCleaningSettings(const FMeshMakeCleanGeometrySettings& CleaningSettings)
+{
+	UpdateSettingsSourceNodeValue(*Graph, CleanMeshSettingsNode, CleaningSettings);
+	CurrentCleanMeshSettings = CleaningSettings;
+}
+
 void FGenerateMeshLODGraph::UpdateSimplifySettings(const FMeshSimplifySettings& SimplifySettings)
 {
 	UpdateSettingsSourceNodeValue(*Graph, SimplifySettingsNode, SimplifySettings);
 	CurrentSimplifySettings = SimplifySettings;
 }
 
+void FGenerateMeshLODGraph::UpdateNormalsSettings(const FMeshNormalsSettings& NormalsSettings)
+{
+	UpdateSettingsSourceNodeValue(*Graph, NormalsSettingsNode, NormalsSettings);
+	CurrentNormalsSettings = NormalsSettings;
+}
+
 void FGenerateMeshLODGraph::UpdateAutoUVSettings(const UE::GeometryFlow::FMeshAutoGenerateUVsSettings& AutoUVSettings)
 {
 	UpdateSettingsSourceNodeValue(*Graph, AutoUVSettingsNode, AutoUVSettings);
@@ -248,10 +270,38 @@ void FGenerateMeshLODGraph::EvaluateResult(
 
 
 
+static bool IsSingleMaterialMesh(const FDynamicMesh3& Mesh, int32& UniqueMaterialID)
+{
+	if (Mesh.HasAttributes() == false) return true;
+	if (Mesh.Attributes()->HasMaterialID() == false) return true;
+
+	const FDynamicMeshMaterialAttribute* MaterialIDs = Mesh.Attributes()->GetMaterialID();
+	UniqueMaterialID = 0;
+	bool bConstantInitialized = false;
+	for (int32 TriangleID : Mesh.TriangleIndicesItr())
+	{
+		int32 MaterialID = MaterialIDs->GetValue(TriangleID);
+		if (bConstantInitialized && MaterialID != UniqueMaterialID)
+		{
+			return false;
+		}
+		if (!bConstantInitialized)
+		{
+			bConstantInitialized = true;
+			UniqueMaterialID = MaterialID;
+		}
+	}
+	return true;
+}
+
 
 
-void FGenerateMeshLODGraph::BuildGraph()
+void FGenerateMeshLODGraph::BuildGraph(const FDynamicMesh3* SourceMeshHint)
 {
+	// precompute hints (async?)
+	int32 UniqueMaterialID = 0;
+	bool bIsSingleMaterialMesh = (SourceMeshHint != nullptr) ? IsSingleMaterialMesh(*SourceMeshHint, UniqueMaterialID) : false;
+
 	Graph = MakeUnique<FGraph>();
 
 	MeshSourceNode = Graph->AddNodeOfType<FDynamicMeshSourceNode>(TEXT("SourceMesh"));
@@ -291,14 +341,38 @@ void FGenerateMeshLODGraph::BuildGraph()
 	MorphologySettingsNode = Graph->AddNodeOfType<FVoxClosureSettingsSourceNode>(TEXT("ClosureSettings"));
 	Graph->InferConnection(MorphologySettingsNode, MorphologyNode);
 
-	// todo: if we only have one material ID we can skip this...
-	FGraph::FHandle MatIDTransferNode = Graph->AddNodeOfType<FTransferMeshMaterialIDsNode>(TEXT("TransferMaterialIDs"));
-	Graph->AddConnection(MeshSourceNode, FDynamicMeshSourceNode::OutParamValue(), MatIDTransferNode, FTransferMeshMaterialIDsNode::InParamMaterialSourceMesh());
-	Graph->InferConnection(MorphologyNode, MatIDTransferNode);
+	CleanMeshNode = Graph->AddNodeOfType<FMeshMakeCleanGeometryNode>(TEXT("Cleaner"));
+	Graph->InferConnection(ThickenNode, CleanMeshNode);
+	CleanMeshSettingsNode = Graph->AddNodeOfType<FMeshMakeCleanGeometrySettingsSourceNode>(TEXT("CleanerSettings"));
+	Graph->InferConnection(CleanMeshSettingsNode, CleanMeshNode);
+
+	// integer param that controls which mesh generator will be used
+	MeshGeneratorSwitchNode = Graph->AddNodeOfType<FMeshGeneratorSwitchNode>(TEXT("MeshGeneratorSwitch"));
+	Graph->AddConnection(SolidifyNode, FSolidifyMeshNode::OutParamResultMesh(), MeshGeneratorSwitchNode, FMeshGeneratorSwitchNode::InParamValue(0));
+	Graph->AddConnection(MorphologyNode, FVoxClosureMeshNode::OutParamResultMesh(), MeshGeneratorSwitchNode, FMeshGeneratorSwitchNode::InParamValue(1));
+	Graph->AddConnection(CleanMeshNode, FMeshMakeCleanGeometryNode::OutParamResultMesh(), MeshGeneratorSwitchNode, FMeshGeneratorSwitchNode::InParamValue(2));
+	
+	// helper to separate next block from changes above
+	FGraph::FHandle FinalMeshOutputNode = MeshGeneratorSwitchNode;
+
+	// skip MaterialID projection/transfer if the input mesh only had one material
+	FGraph::FHandle SimplifyPhaseInputNode = FinalMeshOutputNode;
+	if (bIsSingleMaterialMesh == false)
+	{
+		FGraph::FHandle MatIDTransferNode = Graph->AddNodeOfType<FTransferMeshMaterialIDsNode>(TEXT("TransferMaterialIDs"));
+		Graph->AddConnection(MeshSourceNode, FDynamicMeshSourceNode::OutParamValue(), MatIDTransferNode, FTransferMeshMaterialIDsNode::InParamMaterialSourceMesh());
+		Graph->InferConnection(FinalMeshOutputNode, MatIDTransferNode);
+		SimplifyPhaseInputNode = MatIDTransferNode;
+	}
+	else
+	{
+		UE_LOG(LogGeometry, Warning, TEXT("AutoLOD: applying single-material optimizations"));
+	}
 
-	// need to compute valid normals before Simplify, Morphology node does not necessarily do it
-	FGraph::FHandle PerVertexNormalsNode = Graph->AddNodeOfType<FComputeMeshPerVertexOverlayNormalsNode>(TEXT("PerVertexNormals"));
-	Graph->InferConnection(MatIDTransferNode, PerVertexNormalsNode);
+	// need to compute valid normals before Simplify
+	//FGraph::FHandle PerVertexNormalsNode = Graph->AddNodeOfType<FComputeMeshPerVertexOverlayNormalsNode>(TEXT("PerVertexOverlayNormals"));
+	FGraph::FHandle PerVertexNormalsNode = Graph->AddNodeOfType<FComputeMeshPerVertexNormalsNode>(TEXT("PerVertexNormals"));
+	Graph->InferConnection(SimplifyPhaseInputNode, PerVertexNormalsNode);
 
 	SimplifyNode = Graph->AddNodeOfType<FSimplifyMeshNode>(TEXT("Simplify"));
 	Graph->InferConnection(PerVertexNormalsNode, SimplifyNode);
@@ -396,6 +470,8 @@ void FGenerateMeshLODGraph::BuildGraph()
 	// parameters
 	//
 
+	CurrentCoreMeshGeneratorMode = ECoreMeshGeneratorMode::SolidifyAndClose;
+	UpdateSwitchNodeInputIndex<FMeshGeneratorSwitchNode>(*Graph, MeshGeneratorSwitchNode, static_cast<int32>(CurrentCoreMeshGeneratorMode) );
 
 	FIndexSets IgnoreGroupsForDelete;
 	IgnoreGroupsForDelete.AppendSet({ 0 });
@@ -414,6 +490,11 @@ void FGenerateMeshLODGraph::BuildGraph()
 	MorphologySettings.Distance = 5.0;
 	UpdateMorphologySettings(MorphologySettings);
 
+	FMeshMakeCleanGeometrySettings CleanerSettings;
+	CleanerSettings.bClearMaterialIDs = bIsSingleMaterialMesh;
+	CleanerSettings.bOutputOverlayVertexNormals = false;
+	UpdateMeshCleaningSettings(CleanerSettings);
+
 	// TODO: can use simpler settings if there is only one material...
 	FMeshSimplifySettings SimplifySettings;
 	SimplifySettings.bDiscardAttributes = false;
@@ -428,8 +509,8 @@ void FGenerateMeshLODGraph::BuildGraph()
 
 	FMeshNormalsSettings NormalsSettings;
 	NormalsSettings.NormalsType = EComputeNormalsType::FromFaceAngleThreshold;
-	NormalsSettings.AngleThresholdDeg = 45.0;
-	UpdateSettingsSourceNodeValue(*Graph, NormalsSettingsNode, NormalsSettings);
+	NormalsSettings.AngleThresholdDeg = 60.0;
+	UpdateNormalsSettings(NormalsSettings);
 
 	FMeshAutoGenerateUVsSettings AutoUVSettings;
 	AutoUVSettings.Method = EAutoUVMethod::PatchBuilder;
@@ -458,7 +539,7 @@ void FGenerateMeshLODGraph::BuildGraph()
 
 
 	FMeshMakeBakingCacheSettings BakeCacheSettings;
-	BakeCacheSettings.Dimensions = FImageDimensions(512, 512);
+	BakeCacheSettings.Dimensions = FImageDimensions(1024, 1024);
 	BakeCacheSettings.Thickness = 5.0;
 	UpdateBakeCacheSettings(BakeCacheSettings);
 
@@ -473,6 +554,7 @@ void FGenerateMeshLODGraph::BuildGraph()
 	UpdateCollisionGroupLayerName(CollisionGroupLayerName);
 
 	FGenerateSimpleCollisionSettings GenSimpleCollisionSettings;
+	GenSimpleCollisionSettings.Type = UE::GeometryFlow::ESimpleCollisionGeometryType::AlignedBoxes;
 	UpdateGenerateSimpleCollisionSettings(GenSimpleCollisionSettings);
 
 	TArray<float> Weights;
diff --git a/Engine/Plugins/Experimental/MeshLODToolset/Source/MeshLODToolset/Private/Graphs/GenerateStaticMeshLODProcess.cpp b/Engine/Plugins/Experimental/MeshLODToolset/Source/MeshLODToolset/Private/Graphs/GenerateStaticMeshLODProcess.cpp
--- a/Engine/Plugins/Experimental/MeshLODToolset/Source/MeshLODToolset/Private/Graphs/GenerateStaticMeshLODProcess.cpp
+++ b/Engine/Plugins/Experimental/MeshLODToolset/Source/MeshLODToolset/Private/Graphs/GenerateStaticMeshLODProcess.cpp
@@ -536,7 +536,7 @@ int UGenerateStaticMeshLODProcess::SelectTextureToBake(const TArray<FTextureInfo
 bool UGenerateStaticMeshLODProcess::InitializeGenerator()
 {
 	Generator = MakeUnique<FGenerateMeshLODGraph>();
-	Generator->BuildGraph();
+	Generator->BuildGraph( & this->SourceMesh );
 
 	// initialize source textures that need to be baked
 	SourceTextureToDerivedTexIndex.Reset();
@@ -606,6 +606,10 @@ bool UGenerateStaticMeshLODProcess::InitializeGenerator()
 	CurrentSettings_Simplify.TargetPercentage = Generator->GetCurrentSimplifySettings().TargetFraction * 100.0f;
 	CurrentSettings_Simplify.Tolerance = Generator->GetCurrentSimplifySettings().GeometricTolerance;
 
+	CurrentSettings_Normals.Method = static_cast<EGenerateStaticMeshLODProcess_NormalsMethod>(
+		FGenerateStaticMeshLODProcess_NormalsSettings::MapMethodType((int32)Generator->GetCurrentNormalsSettings().NormalsType, false));
+	CurrentSettings_Normals.Angle = Generator->GetCurrentNormalsSettings().AngleThresholdDeg;
+
 	CurrentSettings_Texture.BakeResolution = (EGenerateStaticMeshLODBakeResolution)Generator->GetCurrentBakeCacheSettings().Dimensions.GetWidth();
 	CurrentSettings_Texture.BakeThickness = Generator->GetCurrentBakeCacheSettings().Thickness;
 
@@ -635,6 +639,11 @@ bool UGenerateStaticMeshLODProcess::InitializeGenerator()
 
 void UGenerateStaticMeshLODProcess::UpdateSettings(const FGenerateStaticMeshLODProcessSettings& NewCombinedSettings)
 {
+	if (NewCombinedSettings.MeshGenerator != CurrentSettings.MeshGenerator)
+	{
+		Generator->UpdateCoreMeshGeneratorMode(static_cast<FGenerateMeshLODGraph::ECoreMeshGeneratorMode>(NewCombinedSettings.MeshGenerator));
+	}
+
 	bool bSharedVoxelResolutionChanged = (NewCombinedSettings.SolidifyVoxelResolution != CurrentSettings.SolidifyVoxelResolution);
 	if ( bSharedVoxelResolutionChanged
 		|| (NewCombinedSettings.WindingThreshold != CurrentSettings.WindingThreshold))
@@ -717,6 +726,50 @@ void UGenerateStaticMeshLODProcess::UpdateSimplifySettings(const FGenerateStatic
 
 
 
+int32 FGenerateStaticMeshLODProcess_NormalsSettings::MapMethodType(int32 From, bool bProcessToGraph)
+{
+	if (bProcessToGraph)
+	{
+		switch (static_cast<EGenerateStaticMeshLODProcess_NormalsMethod>(From))
+		{
+			default:
+			case EGenerateStaticMeshLODProcess_NormalsMethod::FromAngleThreshold: return (int32)UE::GeometryFlow::EComputeNormalsType::FromFaceAngleThreshold;
+			case EGenerateStaticMeshLODProcess_NormalsMethod::PerVertex: return (int32)UE::GeometryFlow::EComputeNormalsType::PerVertex;
+			case EGenerateStaticMeshLODProcess_NormalsMethod::PerTriangle: return (int32)UE::GeometryFlow::EComputeNormalsType::PerTriangle;
+		}
+	}
+	else
+	{
+		switch (static_cast<UE::GeometryFlow::EComputeNormalsType>(From))
+		{
+			default:
+			case UE::GeometryFlow::EComputeNormalsType::FromFaceAngleThreshold: return (int32)EGenerateStaticMeshLODProcess_NormalsMethod::FromAngleThreshold;
+			case UE::GeometryFlow::EComputeNormalsType::PerVertex: return (int32)EGenerateStaticMeshLODProcess_NormalsMethod::PerVertex;
+			case UE::GeometryFlow::EComputeNormalsType::PerTriangle: return (int32)EGenerateStaticMeshLODProcess_NormalsMethod::PerTriangle;
+		}
+	}
+}
+
+
+void UGenerateStaticMeshLODProcess::UpdateNormalsSettings(const FGenerateStaticMeshLODProcess_NormalsSettings& NewCombinedSettings)
+{
+	if (NewCombinedSettings.Method != CurrentSettings_Normals.Method || 
+		NewCombinedSettings.Angle != CurrentSettings_Normals.Angle )
+	{
+		UE::GeometryFlow::FMeshNormalsSettings NewNormalsSettings = Generator->GetCurrentNormalsSettings();
+
+		NewNormalsSettings.NormalsType = static_cast<UE::GeometryFlow::EComputeNormalsType>(
+			FGenerateStaticMeshLODProcess_NormalsSettings::MapMethodType((int32)NewCombinedSettings.Method, true));
+		NewNormalsSettings.AngleThresholdDeg = NewCombinedSettings.Angle;
+
+		Generator->UpdateNormalsSettings(NewNormalsSettings);
+	}
+
+	CurrentSettings_Normals = NewCombinedSettings;
+}
+
+
+
 void UGenerateStaticMeshLODProcess::UpdateTextureSettings(const FGenerateStaticMeshLODProcess_TextureSettings& NewCombinedSettings)
 {
 	if ( (NewCombinedSettings.BakeResolution != CurrentSettings_Texture.BakeResolution) ||
diff --git a/Engine/Plugins/Experimental/MeshLODToolset/Source/MeshLODToolset/Private/Tools/GenerateStaticMeshLODAssetTool.cpp b/Engine/Plugins/Experimental/MeshLODToolset/Source/MeshLODToolset/Private/Tools/GenerateStaticMeshLODAssetTool.cpp
--- a/Engine/Plugins/Experimental/MeshLODToolset/Source/MeshLODToolset/Private/Tools/GenerateStaticMeshLODAssetTool.cpp
+++ b/Engine/Plugins/Experimental/MeshLODToolset/Source/MeshLODToolset/Private/Tools/GenerateStaticMeshLODAssetTool.cpp
@@ -76,6 +76,7 @@ namespace GenerateStaticMeshLODAssetLocals
 		FGenerateStaticMeshLODProcess_PreprocessSettings GeneratorSettings_Preprocess;
 		FGenerateStaticMeshLODProcessSettings GeneratorSettings_MeshGeneration;
 		FGenerateStaticMeshLODProcess_SimplifySettings GeneratorSettings_Simplify;
+		FGenerateStaticMeshLODProcess_NormalsSettings GeneratorSettings_Normals;
 		FGenerateStaticMeshLODProcess_TextureSettings GeneratorSettings_Texture;
 		FGenerateStaticMeshLODProcess_UVSettings GeneratorSettings_UV;
 		FGenerateStaticMeshLODProcess_CollisionSettings GeneratorSettings_Collision;
@@ -105,6 +106,7 @@ namespace GenerateStaticMeshLODAssetLocals
 				GenerateProcess->UpdatePreprocessSettings(GeneratorSettings_Preprocess);
 				GenerateProcess->UpdateSettings(GeneratorSettings_MeshGeneration);
 				GenerateProcess->UpdateSimplifySettings(GeneratorSettings_Simplify);
+				GenerateProcess->UpdateNormalsSettings(GeneratorSettings_Normals);
 				GenerateProcess->UpdateTextureSettings(GeneratorSettings_Texture);
 				GenerateProcess->UpdateUVSettings(GeneratorSettings_UV);
 				GenerateProcess->UpdateCollisionSettings(GeneratorSettings_Collision);
@@ -165,6 +167,7 @@ namespace GenerateStaticMeshLODAssetLocals
 			Op->GeneratorSettings_Preprocess = AutoLODTool->BasicProperties->Preprocessing;
 			Op->GeneratorSettings_MeshGeneration = AutoLODTool->BasicProperties->MeshGeneration;
 			Op->GeneratorSettings_Simplify = AutoLODTool->BasicProperties->Simplification;
+			Op->GeneratorSettings_Normals = AutoLODTool->BasicProperties->Normals;
 			Op->GeneratorSettings_Texture = AutoLODTool->BasicProperties->TextureBaking;
 			Op->GeneratorSettings_UV = AutoLODTool->BasicProperties->UVGeneration;
 			Op->GeneratorSettings_Collision.CollisionGroupLayerName = AutoLODTool->BasicProperties->CollisionGroupLayerName;
@@ -316,6 +319,7 @@ void UGenerateStaticMeshLODAssetTool::Setup()
 	BasicProperties->Preprocessing = GenerateProcess->GetCurrentPreprocessSettings();
 	BasicProperties->MeshGeneration = GenerateProcess->GetCurrentSettings();
 	BasicProperties->Simplification = GenerateProcess->GetCurrentSimplifySettings();
+	BasicProperties->Normals = GenerateProcess->GetCurrentNormalsSettings();
 	BasicProperties->TextureBaking = GenerateProcess->GetCurrentTextureSettings();
 	BasicProperties->UVGeneration = GenerateProcess->GetCurrentUVSettings();
 	BasicProperties->SimpleCollision = GenerateProcess->GetCurrentCollisionSettings();
@@ -545,6 +549,7 @@ void UGenerateStaticMeshLODAssetTool::RequestPresetAction(EGenerateLODAssetToolP
 		CurrentPreset->Preprocessing = BasicProperties->Preprocessing;
 		CurrentPreset->MeshGeneration = BasicProperties->MeshGeneration;
 		CurrentPreset->Simplification = BasicProperties->Simplification;
+		CurrentPreset->Normals = BasicProperties->Normals;
 		CurrentPreset->TextureBaking = BasicProperties->TextureBaking;
 		CurrentPreset->UVGeneration = BasicProperties->UVGeneration;
 		CurrentPreset->SimpleCollision = BasicProperties->SimpleCollision;
@@ -564,6 +569,7 @@ void UGenerateStaticMeshLODAssetTool::OnPresetSelectionChanged()
 		BasicProperties->Preprocessing = ApplyPreset->Preprocessing;
 		BasicProperties->MeshGeneration = ApplyPreset->MeshGeneration;
 		BasicProperties->Simplification = ApplyPreset->Simplification;
+		BasicProperties->Normals = ApplyPreset->Normals;
 		BasicProperties->TextureBaking = ApplyPreset->TextureBaking;
 		BasicProperties->UVGeneration = ApplyPreset->UVGeneration;
 		BasicProperties->SimpleCollision = ApplyPreset->SimpleCollision;
diff --git a/Engine/Plugins/Experimental/MeshLODToolset/Source/MeshLODToolset/Public/Graphs/GenerateMeshLODGraph.h b/Engine/Plugins/Experimental/MeshLODToolset/Source/MeshLODToolset/Public/Graphs/GenerateMeshLODGraph.h
--- a/Engine/Plugins/Experimental/MeshLODToolset/Source/MeshLODToolset/Public/Graphs/GenerateMeshLODGraph.h
+++ b/Engine/Plugins/Experimental/MeshLODToolset/Source/MeshLODToolset/Public/Graphs/GenerateMeshLODGraph.h
@@ -12,7 +12,9 @@
 
 #include "MeshProcessingNodes/MeshSolidifyNode.h"
 #include "MeshProcessingNodes/MeshVoxMorphologyNode.h"
+#include "MeshProcessingNodes/MeshMakeCleanGeometryNode.h"
 #include "MeshProcessingNodes/MeshSimplifyNode.h"
+#include "MeshProcessingNodes/MeshNormalsNodes.h"
 #include "MeshProcessingNodes/MeshThickenNode.h"
 #include "MeshProcessingNodes/MeshDeleteTrianglesNode.h"
 #include "MeshProcessingNodes/MeshAutoGenerateUVsNode.h"
@@ -26,11 +28,14 @@ struct FMeshLODGraphPreFilterSettings
 	FName FilterGroupLayerName = FName( "PreFilterGroups" );
 };
 
-
 class FGenerateMeshLODGraph
 {
 public:
-	void BuildGraph();
+	/**
+	 * Initialize the LOD generation graph. 
+	 * @param SourceMesh if the SourceMesh is known, some optimizations can be made in the graph, but this is optional
+	 */
+	void BuildGraph(const FDynamicMesh3* SourceMeshHint = nullptr);
 
 	int32 AppendTextureBakeNode(const UE::Geometry::TImageBuilder<FVector4f>& SourceImage, const FString& Identifier);
 
@@ -38,18 +43,35 @@ public:
 
 	void SetSourceMesh(const FDynamicMesh3& SourceMesh);
 
+
 	void UpdatePreFilterSettings(const FMeshLODGraphPreFilterSettings& PreFilterSettings);
 	const FMeshLODGraphPreFilterSettings& GetCurrentPreFilterSettings() const { return CurrentPreFilterSettings; }
 
+
+	enum class ECoreMeshGeneratorMode
+	{
+		Solidify = 0,
+		SolidifyAndClose = 1,
+		SimplifyOnly = 2
+	};
+	void UpdateCoreMeshGeneratorMode(ECoreMeshGeneratorMode NewMode);
+	ECoreMeshGeneratorMode GetCurrentCoreMeshGeneratorMode() const { return CurrentCoreMeshGeneratorMode; }
+
 	void UpdateSolidifySettings(const UE::GeometryFlow::FMeshSolidifySettings& SolidifySettings);
 	const UE::GeometryFlow::FMeshSolidifySettings& GetCurrentSolidifySettings() const { return CurrentSolidifySettings; }
 
 	void UpdateMorphologySettings(const UE::GeometryFlow::FVoxClosureSettings& MorphologySettings);
 	const UE::GeometryFlow::FVoxClosureSettings& GetCurrentMorphologySettings() const { return CurrentMorphologySettings; }
 
+	void UpdateMeshCleaningSettings(const UE::GeometryFlow::FMeshMakeCleanGeometrySettings& CleaningSettings);
+	const UE::GeometryFlow::FMeshMakeCleanGeometrySettings& GetCurrentMeshCleaningSettings() const { return CurrentCleanMeshSettings; }
+
 	void UpdateSimplifySettings(const UE::GeometryFlow::FMeshSimplifySettings& SimplifySettings);
 	const UE::GeometryFlow::FMeshSimplifySettings& GetCurrentSimplifySettings() const { return CurrentSimplifySettings; }
 
+	void UpdateNormalsSettings(const UE::GeometryFlow::FMeshNormalsSettings& NormalsSettings);
+	const UE::GeometryFlow::FMeshNormalsSettings& GetCurrentNormalsSettings() const { return CurrentNormalsSettings; }
+
 	void UpdateAutoUVSettings(const UE::GeometryFlow::FMeshAutoGenerateUVsSettings& AutoUVSettings);
 	const UE::GeometryFlow::FMeshAutoGenerateUVsSettings& GetCurrentAutoUVSettings() const { return CurrentAutoUVSettings; }
 
@@ -88,6 +110,9 @@ protected:
 	UE::GeometryFlow::FGraph::FHandle FilterTrianglesNode;
 	FMeshLODGraphPreFilterSettings CurrentPreFilterSettings;
 
+	UE::GeometryFlow::FGraph::FHandle MeshGeneratorSwitchNode;
+	ECoreMeshGeneratorMode CurrentCoreMeshGeneratorMode = ECoreMeshGeneratorMode::SolidifyAndClose;
+
 	UE::GeometryFlow::FGraph::FHandle SolidifyNode;
 	UE::GeometryFlow::FGraph::FHandle SolidifySettingsNode;
 	UE::GeometryFlow::FMeshSolidifySettings CurrentSolidifySettings;
@@ -96,12 +121,17 @@ protected:
 	UE::GeometryFlow::FGraph::FHandle MorphologySettingsNode;
 	UE::GeometryFlow::FVoxClosureSettings CurrentMorphologySettings;
 
+	UE::GeometryFlow::FGraph::FHandle CleanMeshNode;
+	UE::GeometryFlow::FGraph::FHandle CleanMeshSettingsNode;
+	UE::GeometryFlow::FMeshMakeCleanGeometrySettings CurrentCleanMeshSettings;
+
 	UE::GeometryFlow::FGraph::FHandle SimplifyNode;
 	UE::GeometryFlow::FGraph::FHandle SimplifySettingsNode;
 	UE::GeometryFlow::FMeshSimplifySettings CurrentSimplifySettings;
 
 	UE::GeometryFlow::FGraph::FHandle NormalsNode;
 	UE::GeometryFlow::FGraph::FHandle NormalsSettingsNode;
+	UE::GeometryFlow::FMeshNormalsSettings CurrentNormalsSettings;
 
 	UE::GeometryFlow::FGraph::FHandle AutoUVNode;
 	UE::GeometryFlow::FGraph::FHandle AutoUVSettingsNode;
diff --git a/Engine/Plugins/Experimental/MeshLODToolset/Source/MeshLODToolset/Public/Graphs/GenerateStaticMeshLODProcess.h b/Engine/Plugins/Experimental/MeshLODToolset/Source/MeshLODToolset/Public/Graphs/GenerateStaticMeshLODProcess.h
--- a/Engine/Plugins/Experimental/MeshLODToolset/Source/MeshLODToolset/Public/Graphs/GenerateStaticMeshLODProcess.h
+++ b/Engine/Plugins/Experimental/MeshLODToolset/Source/MeshLODToolset/Public/Graphs/GenerateStaticMeshLODProcess.h
@@ -20,6 +20,14 @@ using UE::Geometry::FDynamicMesh3;
 
 
 
+UENUM()
+enum class EGenerateStaticMeshLODProcess_MeshGeneratorModes : uint8
+{
+	// note: must keep in sync with FGenerateMeshLODGraph::ECoreMeshGeneratorMode
+	Solidify = 0,
+	SolidifyAndClose = 1,
+	CleanAndSimplify = 2
+};
 
 
 USTRUCT()
@@ -27,6 +35,9 @@ struct FGenerateStaticMeshLODProcessSettings
 {
 	GENERATED_BODY()
 
+	UPROPERTY(EditAnywhere, Category = MeshGenerator, meta = (DisplayName="Mesh Generator"))
+	EGenerateStaticMeshLODProcess_MeshGeneratorModes MeshGenerator = EGenerateStaticMeshLODProcess_MeshGeneratorModes::SolidifyAndClose;
+
 	// Solidify settings
 
 	/** Target number of voxels along the maximum dimension for Solidify operation */
@@ -107,6 +118,35 @@ struct FGenerateStaticMeshLODProcess_SimplifySettings
 
 
 
+
+UENUM()
+enum class EGenerateStaticMeshLODProcess_NormalsMethod : uint8
+{
+	FromAngleThreshold = 0,
+	PerVertex = 1,
+	PerTriangle = 2
+};
+
+
+USTRUCT()
+struct FGenerateStaticMeshLODProcess_NormalsSettings
+{
+	GENERATED_BODY()
+
+	/** Type of Normals to generate */
+	UPROPERTY(EditAnywhere, Category = Normals, meta = (DisplayName = "Normals Method"))
+	EGenerateStaticMeshLODProcess_NormalsMethod Method = EGenerateStaticMeshLODProcess_NormalsMethod::FromAngleThreshold;
+
+	/** Split Normals (ie sharp edge) will be created along mesh edges with opening angles above this threshold */
+	UPROPERTY(EditAnywhere, Category = Normals, meta = (UIMin = "0", UIMax = "180", ClampMin = "0", ClampMax = "180", EditConditionHides, EditCondition = "Method == EGenerateStaticMeshLODProcess_NormalsMethod::FromAngleThreshold"))
+	float Angle = 60.0;
+
+	static int32 MapMethodType(int32 From, bool bProcessToGraph);
+};
+
+
+
+
 UENUM()
 enum class EGenerateStaticMeshLODProcess_AutoUVMethod : uint8
 {
@@ -122,15 +162,15 @@ struct FGenerateStaticMeshLODProcess_UVSettings_PatchBuilder
 	GENERATED_BODY()
 
 	/** This parameter controls alignment of the initial patches to creases in the mesh */
-	UPROPERTY(EditAnywhere, Category = "PatchBuilder", meta = (UIMin = "0.1", UIMax = "2.0", ClampMin = "0.01", ClampMax = "100.0", EditConditionHides, EditCondition = "UVMethod == EGenerateStaticMeshLODProcess_AutoUVMethod::PatchBuilder"))
+	UPROPERTY(EditAnywhere, Category = "PatchBuilder", meta = (UIMin = "0.1", UIMax = "2.0", ClampMin = "0.01", ClampMax = "100.0"))
 	float CurvatureAlignment = 1.0f;
 
 	/** Number of smoothing steps to apply in the ExpMap UV Generation method (Smoothing slightly increases distortion but produces more stable results) */
-	UPROPERTY(EditAnywhere, Category = "PatchBuilder", AdvancedDisplay, meta = (UIMin = "0", UIMax = "25", ClampMin = "0", ClampMax = "1000", EditConditionHides, EditCondition = "UVMethod == EGenerateStaticMeshLODProcess_AutoUVMethod::PatchBuilder"))
+	UPROPERTY(EditAnywhere, Category = "PatchBuilder", AdvancedDisplay, meta = (UIMin = "0", UIMax = "25", ClampMin = "0", ClampMax = "1000"))
 	int SmoothingSteps = 5;
 
 	/** Smoothing parameter, larger values result in faster smoothing in each step */
-	UPROPERTY(EditAnywhere, Category = "PatchBuilder", AdvancedDisplay, meta = (UIMin = "0", UIMax = "1.0", ClampMin = "0", ClampMax = "1.0", EditConditionHides, EditCondition = "UVMethod == EGenerateStaticMeshLODProcess_AutoUVMethod::PatchBuilder"))
+	UPROPERTY(EditAnywhere, Category = "PatchBuilder", AdvancedDisplay, meta = (UIMin = "0", UIMax = "1.0", ClampMin = "0", ClampMax = "1.0"))
 	float SmoothingAlpha = 0.25f;
 };
 
@@ -189,7 +229,7 @@ struct FGenerateStaticMeshLODProcess_TextureSettings
 
 	/** Resolution for normal map and texture baking */
 	UPROPERTY(EditAnywhere, Category = Baking , meta = (DisplayName = "Bake Image Res"))
-	EGenerateStaticMeshLODBakeResolution BakeResolution = EGenerateStaticMeshLODBakeResolution::Resolution512;
+	EGenerateStaticMeshLODBakeResolution BakeResolution = EGenerateStaticMeshLODBakeResolution::Resolution1024;
 
 	/** How far away from the output mesh to search for input mesh during baking */
 	UPROPERTY(EditAnywhere, Category = Baking, meta = (UIMin = "0", UIMax = "100", ClampMin = "0", ClampMax = "1000", DisplayName = "Bake Thickness"))
@@ -249,31 +289,31 @@ struct FGenerateStaticMeshLODProcess_CollisionSettings
 
 	/** Target triangle count for each convex hull after simplification */
 	UPROPERTY(EditAnywhere, Category = Collision, meta = (NoSpinbox = "true", DisplayName = "Convex Tri Count",
-														  EditCondition = "CollisionType == EGenerateStaticMeshLODSimpleCollisionGeometryType::ConvexHulls"))
+														  EditConditionHides, EditCondition = "CollisionType == EGenerateStaticMeshLODSimpleCollisionGeometryType::ConvexHulls"))
 	int ConvexTriangleCount = 50;
 
 	/** Whether to subsample input vertices using a regular grid before computing the convex hull */
-	UPROPERTY(EditAnywhere, Category = Collision, meta = (EditCondition = "CollisionType == EGenerateStaticMeshLODSimpleCollisionGeometryType::ConvexHulls"))
+	UPROPERTY(EditAnywhere, Category = Collision, meta = (EditConditionHides, EditCondition = "CollisionType == EGenerateStaticMeshLODSimpleCollisionGeometryType::ConvexHulls"))
 	bool bPrefilterVertices = true;
 
 	/** Grid resolution (along the maximum-length axis) for subsampling before computing the convex hull */
-	UPROPERTY(EditAnywhere, Category = Collision, meta = (NoSpinbox = "true", EditCondition = "CollisionType == EGenerateStaticMeshLODSimpleCollisionGeometryType::ConvexHulls && bPrefilterVertices", UIMin = 4, UIMax = 100))
+	UPROPERTY(EditAnywhere, Category = Collision, meta = (NoSpinbox = "true", EditConditionHides, EditCondition = "CollisionType == EGenerateStaticMeshLODSimpleCollisionGeometryType::ConvexHulls && bPrefilterVertices", UIMin = 4, UIMax = 100))
 	int PrefilterGridResolution = 10;
 
 
 	// Swept Convex Hull Settings
 
 	/** Whether to simplify polygons used for swept convex hulls */
-	UPROPERTY(EditAnywhere, Category = Collision, meta = (EditCondition = "CollisionType == EGenerateStaticMeshLODSimpleCollisionGeometryType::SweptHulls"))
+	UPROPERTY(EditAnywhere, Category = Collision, meta = (EditConditionHides, EditCondition = "CollisionType == EGenerateStaticMeshLODSimpleCollisionGeometryType::SweptHulls"))
 	bool bSimplifyPolygons = true;
 
 	/** Target minumum edge length for simplified swept convex hulls */
 	UPROPERTY(EditAnywhere, Category = Collision, meta = (NoSpinbox = "true", UIMin = "0", UIMax = "10", ClampMin = "0", ClampMax = "100000",
-														  EditCondition = "CollisionType == EGenerateStaticMeshLODSimpleCollisionGeometryType::SweptHulls"))
+														  EditConditionHides, EditCondition = "CollisionType == EGenerateStaticMeshLODSimpleCollisionGeometryType::SweptHulls"))
 	float HullTolerance = 0.1;
 
 	/** Which axis to sweep along when computing swept convex hulls */
-	UPROPERTY(EditAnywhere, Category = Collision, meta = (EditCondition = "CollisionType == EGenerateStaticMeshLODSimpleCollisionGeometryType::SweptHulls"))
+	UPROPERTY(EditAnywhere, Category = Collision, meta = (EditConditionHides, EditCondition = "CollisionType == EGenerateStaticMeshLODSimpleCollisionGeometryType::SweptHulls"))
 	EGenerateStaticMeshLODProjectedHullAxisMode SweepAxis = EGenerateStaticMeshLODProjectedHullAxisMode::SmallestVolume;
 };
 
@@ -315,6 +355,9 @@ public:
 	const FGenerateStaticMeshLODProcess_SimplifySettings& GetCurrentSimplifySettings() const { return CurrentSettings_Simplify; }
 	void UpdateSimplifySettings(const FGenerateStaticMeshLODProcess_SimplifySettings& NewSettings);
 
+	const FGenerateStaticMeshLODProcess_NormalsSettings& GetCurrentNormalsSettings() const { return CurrentSettings_Normals; }
+	void UpdateNormalsSettings(const FGenerateStaticMeshLODProcess_NormalsSettings& NewSettings);
+
 	const FGenerateStaticMeshLODProcess_TextureSettings& GetCurrentTextureSettings() const { return CurrentSettings_Texture; }
 	void UpdateTextureSettings(const FGenerateStaticMeshLODProcess_TextureSettings& NewSettings);
 
@@ -529,6 +572,7 @@ protected:
 	FGenerateStaticMeshLODProcess_PreprocessSettings CurrentSettings_Preprocess;
 	FGenerateStaticMeshLODProcessSettings CurrentSettings;
 	FGenerateStaticMeshLODProcess_SimplifySettings CurrentSettings_Simplify;
+	FGenerateStaticMeshLODProcess_NormalsSettings CurrentSettings_Normals;
 	FGenerateStaticMeshLODProcess_TextureSettings CurrentSettings_Texture;
 	FGenerateStaticMeshLODProcess_UVSettings CurrentSettings_UV;
 	FGenerateStaticMeshLODProcess_CollisionSettings CurrentSettings_Collision;
diff --git a/Engine/Plugins/Experimental/MeshLODToolset/Source/MeshLODToolset/Public/Tools/GenerateStaticMeshLODAssetTool.h b/Engine/Plugins/Experimental/MeshLODToolset/Source/MeshLODToolset/Public/Tools/GenerateStaticMeshLODAssetTool.h
--- a/Engine/Plugins/Experimental/MeshLODToolset/Source/MeshLODToolset/Public/Tools/GenerateStaticMeshLODAssetTool.h
+++ b/Engine/Plugins/Experimental/MeshLODToolset/Source/MeshLODToolset/Public/Tools/GenerateStaticMeshLODAssetTool.h
@@ -135,6 +135,9 @@ public:
 	UPROPERTY(EditAnywhere, Category = "Generator Configuration", meta = (ExpandByDefault))
 	FGenerateStaticMeshLODProcess_SimplifySettings Simplification;
 
+	UPROPERTY(EditAnywhere, Category = "Generator Configuration", meta = (ExpandByDefault))
+	FGenerateStaticMeshLODProcess_NormalsSettings Normals;
+
 	UPROPERTY(EditAnywhere, Category = "Generator Configuration", meta = (ExpandByDefault))
 	FGenerateStaticMeshLODProcess_TextureSettings TextureBaking;
 
diff --git a/Engine/Plugins/Experimental/MeshLODToolset/Source/MeshLODToolset/Public/Tools/LODGenerationSettingsAsset.h b/Engine/Plugins/Experimental/MeshLODToolset/Source/MeshLODToolset/Public/Tools/LODGenerationSettingsAsset.h
--- a/Engine/Plugins/Experimental/MeshLODToolset/Source/MeshLODToolset/Public/Tools/LODGenerationSettingsAsset.h
+++ b/Engine/Plugins/Experimental/MeshLODToolset/Source/MeshLODToolset/Public/Tools/LODGenerationSettingsAsset.h
@@ -36,6 +36,9 @@ public:
 	UPROPERTY(EditAnywhere, Category=Settings)
 	FGenerateStaticMeshLODProcess_SimplifySettings Simplification;
 
+	UPROPERTY(EditAnywhere, Category=Settings)
+	FGenerateStaticMeshLODProcess_NormalsSettings Normals;
+
 	UPROPERTY(EditAnywhere, Category=Settings)
 	FGenerateStaticMeshLODProcess_TextureSettings TextureBaking;
 
