AutoLOD: add asset metadata for generated textures/MICs/meshes
#rnx

diff --git a/Engine/Plugins/Experimental/MeshLODToolset/Source/MeshLODToolset/Private/Graphs/GenerateStaticMeshLODProcess.cpp b/Engine/Plugins/Experimental/MeshLODToolset/Source/MeshLODToolset/Private/Graphs/GenerateStaticMeshLODProcess.cpp
--- a/Engine/Plugins/Experimental/MeshLODToolset/Source/MeshLODToolset/Private/Graphs/GenerateStaticMeshLODProcess.cpp
+++ b/Engine/Plugins/Experimental/MeshLODToolset/Source/MeshLODToolset/Private/Graphs/GenerateStaticMeshLODProcess.cpp
@@ -23,6 +23,7 @@
 #include "AssetRegistryModule.h"
 #include "IAssetTools.h"
 #include "FileHelpers.h"
+#include "UObject/Metadata.h"
 
 #include "Engine/Engine.h"
 #include "Editor.h"
@@ -1158,6 +1159,7 @@ bool UGenerateStaticMeshLODProcess::WriteDerivedAssetData()
 	AllDerivedTextures.Reset();
 
 	constexpr bool bCreatingNewStaticMeshAsset = true;
+	DerivedAssetGUIDKey = FGuid::NewGuid().ToString(EGuidFormats::UniqueObjectGuid).ToUpper();
 
 	WriteDerivedTextures(bCreatingNewStaticMeshAsset);
 
@@ -1177,6 +1179,7 @@ void UGenerateStaticMeshLODProcess::UpdateSourceAsset(bool bSetNewHDSourceAsset)
 	AllDerivedTextures.Reset();
 
 	constexpr bool bCreatingNewStaticMeshAsset = false;
+	DerivedAssetGUIDKey = FGuid::NewGuid().ToString(EGuidFormats::UniqueObjectGuid).ToUpper();
 
 	WriteDerivedTextures(bCreatingNewStaticMeshAsset);
 
@@ -1421,6 +1424,12 @@ bool UGenerateStaticMeshLODProcess::WriteDerivedTexture(UTexture2D* DerivedTextu
 	DerivedTexture->SetFlags(RF_Public | RF_Standalone | RF_Transactional);
 	// do we need to Modify() it? we are not doing any undo/redo
 	DerivedTexture->Modify();
+
+	// set metadata tag so we can identify this material later
+	DerivedTexture->GetOutermost()->GetMetaData()->SetValue(DerivedTexture, TEXT("StaticMeshLOD.IsBakedTexture"), TEXT("true"));
+	DerivedTexture->GetOutermost()->GetMetaData()->SetValue(DerivedTexture, TEXT("StaticMeshLOD.SourceAssetPath"), *GetSourceAssetPath());
+	DerivedTexture->GetOutermost()->GetMetaData()->SetValue(DerivedTexture, TEXT("StaticMeshLOD.GenerationGUID"), *DerivedAssetGUIDKey);
+
 	DerivedTexture->UpdateResource();
 	DerivedTexture->PostEditChange();		// this may be necessary if any Materials are using this texture
 	DerivedTexture->MarkPackageDirty();
@@ -1509,6 +1518,11 @@ void UGenerateStaticMeshLODProcess::WriteDerivedMaterials(bool bCreatingNewStati
 			}
 		}
 
+		// set metadata tag so we can identify this material later
+		GeneratedMIC->GetOutermost()->GetMetaData()->SetValue(GeneratedMIC, TEXT("StaticMeshLOD.IsGeneratedMaterial"), TEXT("true"));
+		GeneratedMIC->GetOutermost()->GetMetaData()->SetValue(GeneratedMIC, TEXT("StaticMeshLOD.SourceAssetPath"), *GetSourceAssetPath());
+		GeneratedMIC->GetOutermost()->GetMetaData()->SetValue(GeneratedMIC, TEXT("StaticMeshLOD.GenerationGUID"), *DerivedAssetGUIDKey);
+
 		// rewrite texture parameters to new textures
 		UpdateMaterialTextureParameters(GeneratedMIC, DerivedMaterialInfo);
 
@@ -1641,6 +1655,9 @@ void UGenerateStaticMeshLODProcess::WriteDerivedStaticMeshAsset()
 	// is this necessary? 
 	GeneratedStaticMesh->CreateNavCollision(/*bIsUpdate=*/true);
 
+	GeneratedStaticMesh->GetOutermost()->GetMetaData()->SetValue(GeneratedStaticMesh, TEXT("StaticMeshLOD.IsGeneratedMesh"), TEXT("true"));
+	GeneratedStaticMesh->GetOutermost()->GetMetaData()->SetValue(GeneratedStaticMesh, TEXT("StaticMeshLOD.SourceAssetPath"), *GetSourceAssetPath());
+	GeneratedStaticMesh->GetOutermost()->GetMetaData()->SetValue(GeneratedStaticMesh, TEXT("StaticMeshLOD.GenerationGUID"), *DerivedAssetGUIDKey);
 
 	// done updating mesh
 	GeneratedStaticMesh->PostEditChange();
@@ -1829,6 +1846,9 @@ void UGenerateStaticMeshLODProcess::UpdateSourceStaticMeshAsset(bool bSetNewHDSo
 	// is this necessary? 
 	SourceStaticMesh->CreateNavCollision(/*bIsUpdate=*/true);
 
+	// save UUID used for generated Assets
+	SourceStaticMesh->GetOutermost()->GetMetaData()->SetValue(SourceStaticMesh, TEXT("StaticMeshLOD.GenerationGUID"), *DerivedAssetGUIDKey);
+
 	GEditor->EndTransaction();
 
 	// done updating mesh
diff --git a/Engine/Plugins/Experimental/MeshLODToolset/Source/MeshLODToolset/Public/Graphs/GenerateStaticMeshLODProcess.h b/Engine/Plugins/Experimental/MeshLODToolset/Source/MeshLODToolset/Public/Graphs/GenerateStaticMeshLODProcess.h
--- a/Engine/Plugins/Experimental/MeshLODToolset/Source/MeshLODToolset/Public/Graphs/GenerateStaticMeshLODProcess.h
+++ b/Engine/Plugins/Experimental/MeshLODToolset/Source/MeshLODToolset/Public/Graphs/GenerateStaticMeshLODProcess.h
@@ -577,6 +577,9 @@ protected:
 	FGenerateStaticMeshLODProcess_UVSettings CurrentSettings_UV;
 	FGenerateStaticMeshLODProcess_CollisionSettings CurrentSettings_Collision;
 	
+	// This value will be stored in Asset Metadata for each generated Asset Texture, Material/MIC, and StaticMesh,
+	// under the key StaticMeshLOD.GenerationGUID.
+	FString DerivedAssetGUIDKey;
 
 	bool WriteDerivedTexture(UTexture2D* SourceTexture, UTexture2D* DerivedTexture, bool bCreatingNewStaticMeshAsset);
 	bool WriteDerivedTexture(UTexture2D* DerivedTexture, FString BaseTexName, bool bCreatingNewStaticMeshAsset);
