Add Convex Hull mesh generator option to AutoLOD

diff --git a/Engine/Plugins/Experimental/GeometryFlow/Source/GeometryFlowMeshProcessing/Private/MeshProcessingNodes/GenerateConvexHullMeshNode.cpp b/Engine/Plugins/Experimental/GeometryFlow/Source/GeometryFlowMeshProcessing/Private/MeshProcessingNodes/GenerateConvexHullMeshNode.cpp
new file mode 100644
--- /dev/null
+++ b/Engine/Plugins/Experimental/GeometryFlow/Source/GeometryFlowMeshProcessing/Private/MeshProcessingNodes/GenerateConvexHullMeshNode.cpp
@@ -0,0 +1,31 @@
+// Copyright Epic Games, Inc. All Rights Reserved.
+
+#include "MeshProcessingNodes/GenerateConvexHullMeshNode.h"
+#include "Operations/MeshConvexHull.h"
+#include "Util/ProgressCancel.h"
+
+using namespace UE::Geometry;
+using namespace UE::GeometryFlow;
+
+EGeometryFlowResult FGenerateConvexHullMeshNode::MakeConvexHullMesh(const FDynamicMesh3& MeshIn,
+	const FGenerateConvexHullMeshSettings& Settings, 
+	FDynamicMesh3& MeshOut)
+{
+	FMeshConvexHull Hull(&MeshIn);
+
+	if (Settings.bPrefilterVertices)
+	{
+		FMeshConvexHull::GridSample(MeshIn, Settings.PrefilterGridResolution, Hull.VertexSet);
+	}
+
+	Hull.bPostSimplify = false;		// Mesh can be simplified later
+
+	FProgressCancel* Progress = nullptr;	// TODO: Add ProgressCancel
+	if (Hull.Compute(Progress))
+	{
+		MeshOut = MoveTemp(Hull.ConvexHull);
+	}
+
+	return EGeometryFlowResult::Ok;
+}
+
diff --git a/Engine/Plugins/Experimental/GeometryFlow/Source/GeometryFlowMeshProcessing/Public/MeshProcessingNodes/GenerateConvexHullMeshNode.h b/Engine/Plugins/Experimental/GeometryFlow/Source/GeometryFlowMeshProcessing/Public/MeshProcessingNodes/GenerateConvexHullMeshNode.h
new file mode 100644
--- /dev/null
+++ b/Engine/Plugins/Experimental/GeometryFlow/Source/GeometryFlowMeshProcessing/Public/MeshProcessingNodes/GenerateConvexHullMeshNode.h
@@ -0,0 +1,44 @@
+// Copyright Epic Games, Inc. All Rights Reserved.
+
+#pragma once
+
+#include "GeometryFlowCoreNodes.h"
+#include "MeshProcessingBaseNodes.h"
+
+namespace UE
+{
+namespace GeometryFlow
+{
+
+using namespace UE::Geometry;
+
+struct GEOMETRYFLOWMESHPROCESSING_API FGenerateConvexHullMeshSettings
+{
+	DECLARE_GEOMETRYFLOW_DATA_TYPE_IDENTIFIER(EMeshProcessingDataTypes::GenerateConvexHullMeshSettings);
+	bool bPrefilterVertices = true;
+	int PrefilterGridResolution = 10;
+};
+
+GEOMETRYFLOW_DECLARE_SETTINGS_TYPES(FGenerateConvexHullMeshSettings, GenerateConvexHullMesh);
+
+
+class GEOMETRYFLOWMESHPROCESSING_API FGenerateConvexHullMeshNode : public TProcessMeshWithSettingsBaseNode<FGenerateConvexHullMeshSettings>
+{
+
+public:
+
+	virtual void ProcessMesh(
+		const FNamedDataMap& DatasIn,
+		const FGenerateConvexHullMeshSettings& SettingsIn,
+		const FDynamicMesh3& MeshIn,
+		FDynamicMesh3& MeshOut) override
+	{
+		MakeConvexHullMesh(MeshIn, SettingsIn, MeshOut);
+	}
+
+	EGeometryFlowResult MakeConvexHullMesh(const FDynamicMesh3& MeshIn, const FGenerateConvexHullMeshSettings& Settings, FDynamicMesh3& MeshOut);
+
+};
+
+}	// end namespace GeometryFlow
+}	// end namespace UE
diff --git a/Engine/Plugins/Experimental/GeometryFlow/Source/GeometryFlowMeshProcessing/Public/MeshProcessingNodes/MeshProcessingDataTypes.h b/Engine/Plugins/Experimental/GeometryFlow/Source/GeometryFlowMeshProcessing/Public/MeshProcessingNodes/MeshProcessingDataTypes.h
--- a/Engine/Plugins/Experimental/GeometryFlow/Source/GeometryFlowMeshProcessing/Public/MeshProcessingNodes/MeshProcessingDataTypes.h
+++ b/Engine/Plugins/Experimental/GeometryFlow/Source/GeometryFlowMeshProcessing/Public/MeshProcessingNodes/MeshProcessingDataTypes.h
@@ -27,6 +27,7 @@ enum class EMeshProcessingDataTypes
 	SolidifySettings = (int)EDataTypes::BaseMeshProcessingTypes + 100,
 	VoxMorphologyOpSettings = (int)EDataTypes::BaseMeshProcessingTypes + 101,
 	ThickenSettings = (int)EDataTypes::BaseMeshProcessingTypes + 102,
+	GenerateConvexHullMeshSettings = (int)EDataTypes::BaseMeshProcessingTypes + 103,
 
 	SimplifySettings = (int)EDataTypes::BaseMeshProcessingTypes + 110,
 	NormalFlowSettings = (int)EDataTypes::BaseMeshProcessingTypes + 111,
diff --git a/Engine/Plugins/Experimental/MeshLODToolset/Source/MeshLODToolset/Private/Graphs/GenerateMeshLODGraph.cpp b/Engine/Plugins/Experimental/MeshLODToolset/Source/MeshLODToolset/Private/Graphs/GenerateMeshLODGraph.cpp
--- a/Engine/Plugins/Experimental/MeshLODToolset/Source/MeshLODToolset/Private/Graphs/GenerateMeshLODGraph.cpp
+++ b/Engine/Plugins/Experimental/MeshLODToolset/Source/MeshLODToolset/Private/Graphs/GenerateMeshLODGraph.cpp
@@ -36,7 +36,7 @@ using namespace UE::Geometry;
 using namespace UE::GeometryFlow;
 
 
-typedef UE::GeometryFlow::TSwitchNode<FDynamicMesh3, 3, (int)EMeshProcessingDataTypes::DynamicMesh> FMeshGeneratorSwitchNode;
+typedef UE::GeometryFlow::TSwitchNode<FDynamicMesh3, 4, (int)EMeshProcessingDataTypes::DynamicMesh> FMeshGeneratorSwitchNode;
 
 
 void FGenerateMeshLODGraph::SetSourceMesh(const FDynamicMesh3& SourceMeshIn)
@@ -82,6 +82,12 @@ void FGenerateMeshLODGraph::UpdateSimplifySettings(const FMeshSimplifySettings&
 	CurrentSimplifySettings = SimplifySettings;
 }
 
+void FGenerateMeshLODGraph::UpdateGenerateConvexHullMeshSettings(const FGenerateConvexHullMeshSettings& GenerateConvexHullMeshSettings)
+{
+	UpdateSettingsSourceNodeValue(*Graph, GenerateConvexHullMeshSettingsNode, GenerateConvexHullMeshSettings);
+	CurrentGenerateConvexHullMeshSettings = GenerateConvexHullMeshSettings;
+}
+
 void FGenerateMeshLODGraph::UpdateNormalsSettings(const FMeshNormalsSettings& NormalsSettings)
 {
 	UpdateSettingsSourceNodeValue(*Graph, NormalsSettingsNode, NormalsSettings);
@@ -346,12 +352,19 @@ void FGenerateMeshLODGraph::BuildGraph(const FDynamicMesh3* SourceMeshHint)
 	CleanMeshSettingsNode = Graph->AddNodeOfType<FMeshMakeCleanGeometrySettingsSourceNode>(TEXT("CleanerSettings"));
 	Graph->InferConnection(CleanMeshSettingsNode, CleanMeshNode);
 
+	GenerateConvexHullMeshNode = Graph->AddNodeOfType<FGenerateConvexHullMeshNode>(TEXT("ConvexHullMesh"));
+	Graph->InferConnection(FilterTrianglesNode, GenerateConvexHullMeshNode);
+	GenerateConvexHullMeshSettingsNode = Graph->AddNodeOfType<FGenerateConvexHullMeshSettingsSourceNode>(TEXT("ConvexHullMeshSettings"));
+	Graph->InferConnection(GenerateConvexHullMeshSettingsNode, GenerateConvexHullMeshNode);
+
+
 	// integer param that controls which mesh generator will be used
 	MeshGeneratorSwitchNode = Graph->AddNodeOfType<FMeshGeneratorSwitchNode>(TEXT("MeshGeneratorSwitch"));
 	Graph->AddConnection(SolidifyNode, FSolidifyMeshNode::OutParamResultMesh(), MeshGeneratorSwitchNode, FMeshGeneratorSwitchNode::InParamValue(0));
 	Graph->AddConnection(MorphologyNode, FVoxClosureMeshNode::OutParamResultMesh(), MeshGeneratorSwitchNode, FMeshGeneratorSwitchNode::InParamValue(1));
 	Graph->AddConnection(CleanMeshNode, FMeshMakeCleanGeometryNode::OutParamResultMesh(), MeshGeneratorSwitchNode, FMeshGeneratorSwitchNode::InParamValue(2));
-	
+	Graph->AddConnection(GenerateConvexHullMeshNode, FGenerateConvexHullMeshNode::OutParamResultMesh(), MeshGeneratorSwitchNode, FMeshGeneratorSwitchNode::InParamValue(3));
+
 	// helper to separate next block from changes above
 	FGraph::FHandle FinalMeshOutputNode = MeshGeneratorSwitchNode;
 
diff --git a/Engine/Plugins/Experimental/MeshLODToolset/Source/MeshLODToolset/Private/Graphs/GenerateStaticMeshLODProcess.cpp b/Engine/Plugins/Experimental/MeshLODToolset/Source/MeshLODToolset/Private/Graphs/GenerateStaticMeshLODProcess.cpp
--- a/Engine/Plugins/Experimental/MeshLODToolset/Source/MeshLODToolset/Private/Graphs/GenerateStaticMeshLODProcess.cpp
+++ b/Engine/Plugins/Experimental/MeshLODToolset/Source/MeshLODToolset/Private/Graphs/GenerateStaticMeshLODProcess.cpp
@@ -665,6 +665,15 @@ void UGenerateStaticMeshLODProcess::UpdateSettings(const FGenerateStaticMeshLODP
 		Generator->UpdateMorphologySettings(NewClosureSettings);
 	}
 
+	if (NewCombinedSettings.bPrefilterVertices != CurrentSettings.bPrefilterVertices ||
+		NewCombinedSettings.PrefilterGridResolution != CurrentSettings.PrefilterGridResolution)
+	{
+		UE::GeometryFlow::FGenerateConvexHullMeshSettings NewConvexHullSettings = Generator->GetCurrentGenerateConvexHullMeshSettings();
+		NewConvexHullSettings.bPrefilterVertices = NewCombinedSettings.bPrefilterVertices;
+		NewConvexHullSettings.PrefilterGridResolution = NewCombinedSettings.PrefilterGridResolution;
+		Generator->UpdateGenerateConvexHullMeshSettings(NewConvexHullSettings);
+	}
+
 	CurrentSettings = NewCombinedSettings;
 }
 
diff --git a/Engine/Plugins/Experimental/MeshLODToolset/Source/MeshLODToolset/Public/Graphs/GenerateMeshLODGraph.h b/Engine/Plugins/Experimental/MeshLODToolset/Source/MeshLODToolset/Public/Graphs/GenerateMeshLODGraph.h
--- a/Engine/Plugins/Experimental/MeshLODToolset/Source/MeshLODToolset/Public/Graphs/GenerateMeshLODGraph.h
+++ b/Engine/Plugins/Experimental/MeshLODToolset/Source/MeshLODToolset/Public/Graphs/GenerateMeshLODGraph.h
@@ -18,6 +18,7 @@
 #include "MeshProcessingNodes/MeshThickenNode.h"
 #include "MeshProcessingNodes/MeshDeleteTrianglesNode.h"
 #include "MeshProcessingNodes/MeshAutoGenerateUVsNode.h"
+#include "MeshProcessingNodes/GenerateConvexHullMeshNode.h"
 #include "DataTypes/MeshImageBakingData.h"
 #include "PhysicsNodes/GenerateSimpleCollisionNode.h"
 
@@ -52,7 +53,8 @@ public:
 	{
 		Solidify = 0,
 		SolidifyAndClose = 1,
-		SimplifyOnly = 2
+		SimplifyOnly = 2,
+		ConvexHull = 3
 	};
 	void UpdateCoreMeshGeneratorMode(ECoreMeshGeneratorMode NewMode);
 	ECoreMeshGeneratorMode GetCurrentCoreMeshGeneratorMode() const { return CurrentCoreMeshGeneratorMode; }
@@ -69,6 +71,9 @@ public:
 	void UpdateSimplifySettings(const UE::GeometryFlow::FMeshSimplifySettings& SimplifySettings);
 	const UE::GeometryFlow::FMeshSimplifySettings& GetCurrentSimplifySettings() const { return CurrentSimplifySettings; }
 
+	void UpdateGenerateConvexHullMeshSettings(const UE::GeometryFlow::FGenerateConvexHullMeshSettings& ConvexHullSettings);
+	const UE::GeometryFlow::FGenerateConvexHullMeshSettings& GetCurrentGenerateConvexHullMeshSettings() const { return CurrentGenerateConvexHullMeshSettings; }
+
 	void UpdateNormalsSettings(const UE::GeometryFlow::FMeshNormalsSettings& NormalsSettings);
 	const UE::GeometryFlow::FMeshNormalsSettings& GetCurrentNormalsSettings() const { return CurrentNormalsSettings; }
 
@@ -129,6 +134,10 @@ protected:
 	UE::GeometryFlow::FGraph::FHandle SimplifySettingsNode;
 	UE::GeometryFlow::FMeshSimplifySettings CurrentSimplifySettings;
 
+	UE::GeometryFlow::FGraph::FHandle GenerateConvexHullMeshNode;
+	UE::GeometryFlow::FGraph::FHandle GenerateConvexHullMeshSettingsNode;
+	UE::GeometryFlow::FGenerateConvexHullMeshSettings CurrentGenerateConvexHullMeshSettings;
+
 	UE::GeometryFlow::FGraph::FHandle NormalsNode;
 	UE::GeometryFlow::FGraph::FHandle NormalsSettingsNode;
 	UE::GeometryFlow::FMeshNormalsSettings CurrentNormalsSettings;
diff --git a/Engine/Plugins/Experimental/MeshLODToolset/Source/MeshLODToolset/Public/Graphs/GenerateStaticMeshLODProcess.h b/Engine/Plugins/Experimental/MeshLODToolset/Source/MeshLODToolset/Public/Graphs/GenerateStaticMeshLODProcess.h
--- a/Engine/Plugins/Experimental/MeshLODToolset/Source/MeshLODToolset/Public/Graphs/GenerateStaticMeshLODProcess.h
+++ b/Engine/Plugins/Experimental/MeshLODToolset/Source/MeshLODToolset/Public/Graphs/GenerateStaticMeshLODProcess.h
@@ -26,7 +26,8 @@ enum class EGenerateStaticMeshLODProcess_MeshGeneratorModes : uint8
 	// note: must keep in sync with FGenerateMeshLODGraph::ECoreMeshGeneratorMode
 	Solidify = 0,
 	SolidifyAndClose = 1,
-	CleanAndSimplify = 2
+	CleanAndSimplify = 2,
+	ConvexHull = 3
 };
 
 
@@ -41,20 +42,34 @@ struct FGenerateStaticMeshLODProcessSettings
 	// Solidify settings
 
 	/** Target number of voxels along the maximum dimension for Solidify operation */
-	UPROPERTY(EditAnywhere, Category = Solidify, meta = (UIMin = "8", UIMax = "1024", ClampMin = "8", ClampMax = "1024", DisplayName="Voxel Resolution"))
+	UPROPERTY(EditAnywhere, Category = Solidify, meta = (UIMin = "8", UIMax = "1024", ClampMin = "8", ClampMax = "1024", DisplayName="Voxel Resolution",
+	EditConditionHides, EditCondition = "MeshGenerator != EGenerateStaticMeshLODProcess_MeshGeneratorModes::ConvexHull"))
 	int SolidifyVoxelResolution = 128;
 
 	/** Winding number threshold to determine what is considered inside the mesh during Solidify */
-	UPROPERTY(EditAnywhere, Category = Solidify, AdvancedDisplay, meta = (UIMin = "0.1", UIMax = ".9", ClampMin = "-10", ClampMax = "10"))
+	UPROPERTY(EditAnywhere, Category = Solidify, AdvancedDisplay, meta = (UIMin = "0.1", UIMax = ".9", ClampMin = "-10", ClampMax = "10",
+	EditConditionHides, EditCondition = "MeshGenerator != EGenerateStaticMeshLODProcess_MeshGeneratorModes::ConvexHull"))
 	float WindingThreshold = 0.5f;
 
 
 	// Morphology settings
 
 	/** Offset distance in the Morpohological Closure operation */
-	UPROPERTY(EditAnywhere, Category = Morphology, meta = (UIMin = "0", UIMax = "100", ClampMin = "0", ClampMax = "1000", DisplayName = "Closure Distance"))
+	UPROPERTY(EditAnywhere, Category = Morphology, meta = (UIMin = "0", UIMax = "100", ClampMin = "0", ClampMax = "1000", DisplayName = "Closure Distance",
+	EditConditionHides, EditCondition = "MeshGenerator != EGenerateStaticMeshLODProcess_MeshGeneratorModes::ConvexHull"))
 	float ClosureDistance = 1.0f;
 
+
+	// Convex Hull Settings
+
+	/** Whether to subsample input vertices using a regular grid before computing the convex hull */
+	UPROPERTY(EditAnywhere, Category = Collision, meta = (EditConditionHides, EditCondition = "MeshGenerator == EGenerateStaticMeshLODProcess_MeshGeneratorModes::ConvexHull"))
+	bool bPrefilterVertices = true;
+
+	/** Grid resolution (along the maximum-length axis) for subsampling before computing the convex hull */
+	UPROPERTY(EditAnywhere, Category = Collision, meta = (NoSpinbox = "true", EditConditionHides, EditCondition = "MeshGenerator == EGenerateStaticMeshLODProcess_MeshGeneratorModes::ConvexHull", UIMin = 4, UIMax = 100))
+	int PrefilterGridResolution = 10;
+
 };
 
 
