AutoLOD: TObject and static analysis fixups
#rnx

diff --git a/Engine/Plugins/Experimental/MeshLODToolset/Source/MeshLODToolset/Private/Graphs/GenerateStaticMeshLODProcess.cpp b/Engine/Plugins/Experimental/MeshLODToolset/Source/MeshLODToolset/Private/Graphs/GenerateStaticMeshLODProcess.cpp
--- a/Engine/Plugins/Experimental/MeshLODToolset/Source/MeshLODToolset/Private/Graphs/GenerateStaticMeshLODProcess.cpp
+++ b/Engine/Plugins/Experimental/MeshLODToolset/Source/MeshLODToolset/Private/Graphs/GenerateStaticMeshLODProcess.cpp
@@ -1527,23 +1527,29 @@ void UGenerateStaticMeshLODProcess::WriteDerivedMaterials(bool bCreatingNewStati
 			}
 		}
 
-		// set metadata tag so we can identify this material later
-		GeneratedMIC->GetOutermost()->GetMetaData()->SetValue(GeneratedMIC, TEXT("StaticMeshLOD.IsGeneratedMaterial"), TEXT("true"));
-		GeneratedMIC->GetOutermost()->GetMetaData()->SetValue(GeneratedMIC, TEXT("StaticMeshLOD.SourceAssetPath"), *GetSourceAssetPath());
-		GeneratedMIC->GetOutermost()->GetMetaData()->SetValue(GeneratedMIC, TEXT("StaticMeshLOD.GenerationGUID"), *DerivedAssetGUIDKey);
+		if (GeneratedMIC != nullptr)
+		{
+			// set metadata tag so we can identify this material later
+			GeneratedMIC->GetOutermost()->GetMetaData()->SetValue(GeneratedMIC, TEXT("StaticMeshLOD.IsGeneratedMaterial"), TEXT("true"));
+			GeneratedMIC->GetOutermost()->GetMetaData()->SetValue(GeneratedMIC, TEXT("StaticMeshLOD.SourceAssetPath"), *GetSourceAssetPath());
+			GeneratedMIC->GetOutermost()->GetMetaData()->SetValue(GeneratedMIC, TEXT("StaticMeshLOD.GenerationGUID"), *DerivedAssetGUIDKey);
 
-		// rewrite texture parameters to new textures
-		UpdateMaterialTextureParameters(GeneratedMIC, DerivedMaterialInfo);
+			// rewrite texture parameters to new textures
+			UpdateMaterialTextureParameters(GeneratedMIC, DerivedMaterialInfo);
+		}
 
 		// update StaticMaterial
 		DerivedMaterialInfo.DerivedMaterial.MaterialInterface = GeneratedMIC;
 		DerivedMaterialInfo.DerivedMaterial.MaterialSlotName = FName(FString::Printf(TEXT("GeneratedMat%d"), mi));
 		DerivedMaterialInfo.DerivedMaterial.ImportedMaterialSlotName = DerivedMaterialInfo.DerivedMaterial.MaterialSlotName;
 
-		if (CurrentSettings_Texture.bCombineTextures && DerivedMultiTextureBakeResult && MultiTextureParameterName.Contains(mi))
+		if (GeneratedMIC != nullptr)
 		{
-			FMaterialParameterInfo ParamInfo(MultiTextureParameterName[mi]);
-			GeneratedMIC->SetTextureParameterValueEditorOnly(ParamInfo, DerivedMultiTextureBakeResult);
+			if (CurrentSettings_Texture.bCombineTextures && DerivedMultiTextureBakeResult && MultiTextureParameterName.Contains(mi))
+			{
+				FMaterialParameterInfo ParamInfo(MultiTextureParameterName[mi]);
+				GeneratedMIC->SetTextureParameterValueEditorOnly(ParamInfo, DerivedMultiTextureBakeResult);
+			}
 		}
 
 	}
diff --git a/Engine/Plugins/Experimental/MeshLODToolset/Source/MeshLODToolset/Private/Tools/GenerateStaticMeshLODAssetTool.cpp b/Engine/Plugins/Experimental/MeshLODToolset/Source/MeshLODToolset/Private/Tools/GenerateStaticMeshLODAssetTool.cpp
--- a/Engine/Plugins/Experimental/MeshLODToolset/Source/MeshLODToolset/Private/Tools/GenerateStaticMeshLODAssetTool.cpp
+++ b/Engine/Plugins/Experimental/MeshLODToolset/Source/MeshLODToolset/Private/Tools/GenerateStaticMeshLODAssetTool.cpp
@@ -533,6 +533,7 @@ void UGenerateStaticMeshLODAssetTool::RequestPresetAction(EGenerateLODAssetToolP
 	{
 		GetToolManager()->DisplayMessage(
 			LOCTEXT("GenerateStaticMeshLODAssetTool_NoPresetSelected", "No Preset Asset is currently set in the Preset Settings"), EToolMessageLevel::UserError);
+		return;
 	}
 
 	if (ActionType == EGenerateLODAssetToolPresetAction::ReadFromPreset)
diff --git a/Engine/Plugins/Experimental/MeshLODToolset/Source/MeshLODToolset/Public/Graphs/GenerateStaticMeshLODProcess.h b/Engine/Plugins/Experimental/MeshLODToolset/Source/MeshLODToolset/Public/Graphs/GenerateStaticMeshLODProcess.h
--- a/Engine/Plugins/Experimental/MeshLODToolset/Source/MeshLODToolset/Public/Graphs/GenerateStaticMeshLODProcess.h
+++ b/Engine/Plugins/Experimental/MeshLODToolset/Source/MeshLODToolset/Public/Graphs/GenerateStaticMeshLODProcess.h
@@ -488,7 +488,7 @@ public:
 protected:
 
 	UPROPERTY()
-	UStaticMesh* SourceStaticMesh;
+	TObjectPtr<UStaticMesh> SourceStaticMesh;
 
 	FString SourceAssetPath;
 	FString SourceAssetFolder;
@@ -568,17 +568,17 @@ protected:
 	// This list is for accumulating derived UTexture2D's created during WriteDerivedTextures(). We have to
 	// maintain uproperty references to these or they may be garbage collected
 	UPROPERTY()
-	TSet<UTexture2D*> AllDerivedTextures;
+	TSet<TObjectPtr<UTexture2D>> AllDerivedTextures;
 
 	// Derived Normal Map
 	UPROPERTY()
-	UTexture2D* DerivedNormalMapTex;
+	TObjectPtr<UTexture2D> DerivedNormalMapTex;
 
 	// For each material participating in multi-texture baking, the parameter name of the texture
 	TMap<int32, FName> MultiTextureParameterName;
 
 	UPROPERTY()
-	UTexture2D* DerivedMultiTextureBakeResult;
+	TObjectPtr<UTexture2D> DerivedMultiTextureBakeResult;
 
 
 	TUniquePtr<FGenerateMeshLODGraph> Generator;			// active LODGenerator Graph
diff --git a/Engine/Plugins/Experimental/MeshLODToolset/Source/MeshLODToolset/Public/Tools/GenerateStaticMeshLODAssetTool.h b/Engine/Plugins/Experimental/MeshLODToolset/Source/MeshLODToolset/Public/Tools/GenerateStaticMeshLODAssetTool.h
--- a/Engine/Plugins/Experimental/MeshLODToolset/Source/MeshLODToolset/Public/Tools/GenerateStaticMeshLODAssetTool.h
+++ b/Engine/Plugins/Experimental/MeshLODToolset/Source/MeshLODToolset/Public/Tools/GenerateStaticMeshLODAssetTool.h
@@ -201,7 +201,7 @@ struct MESHLODTOOLSET_API FGenerateStaticMeshLOD_TextureConfig
 	GENERATED_BODY()
 
 	UPROPERTY(VisibleAnywhere, Category = Texture)
-	UTexture2D* Texture = nullptr;
+	TObjectPtr<UTexture2D> Texture = nullptr;
 
 	UPROPERTY(EditAnywhere, Category = Texture)
 	EGenerateStaticMeshLOD_BakeConstraint Constraint = EGenerateStaticMeshLOD_BakeConstraint::NoConstraint;
@@ -219,7 +219,7 @@ struct MESHLODTOOLSET_API FGenerateStaticMeshLOD_MaterialConfig
 	GENERATED_BODY()
 
 	UPROPERTY(VisibleAnywhere, Category = Material)
-	UMaterialInterface* Material = nullptr;
+	TObjectPtr<UMaterialInterface> Material = nullptr;
 
 	UPROPERTY(EditAnywhere, Category = Texture)
 	EGenerateStaticMeshLOD_BakeConstraint Constraint = EGenerateStaticMeshLOD_BakeConstraint::NoConstraint;
@@ -244,7 +244,7 @@ public:
 	TArray<FGenerateStaticMeshLOD_TextureConfig> Textures;
 
 	UPROPERTY(VisibleAnywhere, Category = "Baked Texture Previews")
-	TArray<UTexture2D*> PreviewTextures;
+	TArray<TObjectPtr<UTexture2D>> PreviewTextures;
 };
 
 
@@ -299,10 +299,10 @@ protected:
 	TObjectPtr<UMeshOpPreviewWithBackgroundCompute> PreviewWithBackgroundCompute = nullptr;
 
 	UPROPERTY()
-	TArray<UTexture2D*> PreviewTextures;
+	TArray<TObjectPtr<UTexture2D>> PreviewTextures;
 
 	UPROPERTY()
-	TArray<UMaterialInterface*> PreviewMaterials;
+	TArray<TObjectPtr<UMaterialInterface>> PreviewMaterials;
 
 
 protected:
diff --git a/Engine/Plugins/Experimental/MeshLODToolset/Source/MeshLODToolset/Public/Tools/LODManagerTool.h b/Engine/Plugins/Experimental/MeshLODToolset/Source/MeshLODToolset/Public/Tools/LODManagerTool.h
--- a/Engine/Plugins/Experimental/MeshLODToolset/Source/MeshLODToolset/Public/Tools/LODManagerTool.h
+++ b/Engine/Plugins/Experimental/MeshLODToolset/Source/MeshLODToolset/Public/Tools/LODManagerTool.h
@@ -189,16 +189,16 @@ public:
 protected:
 
 	UPROPERTY()
-	ULODManagerLODProperties* LODInfoProperties;
+	TObjectPtr<ULODManagerLODProperties> LODInfoProperties;
 
 	UPROPERTY()
-	ULODManagerPreviewLODProperties* LODPreviewProperties;
+	TObjectPtr<ULODManagerPreviewLODProperties> LODPreviewProperties;
 
 	UPROPERTY()
-	ULODManagerHiResSourceModelActions* HiResSourceModelActions;
+	TObjectPtr<ULODManagerHiResSourceModelActions> HiResSourceModelActions;
 
 	UPROPERTY()
-	ULODManagerMaterialActions* MaterialActions;
+	TObjectPtr<ULODManagerMaterialActions> MaterialActions;
 
 
 
@@ -244,10 +244,10 @@ protected:
 	TMap<FString, TUniquePtr<FLODMeshInfo>> LODMeshCache;
 
 	UPROPERTY()
-	UPreviewMesh* LODPreview;
+	TObjectPtr<UPreviewMesh> LODPreview;
 
 	UPROPERTY()
-	UPreviewGeometry* LODPreviewLines;
+	TObjectPtr<UPreviewGeometry> LODPreviewLines;
 
 	bool bPreviewLODValid = false;
 	void UpdatePreviewLOD();
