AutoLOD: add support for 'asset editor mode' in the AutoLOD Tool, hides output-mode option, enable in SME ModelingMode
#rnx

diff --git a/Engine/Plugins/Experimental/MeshLODToolset/Source/MeshLODToolset/Private/Tools/GenerateStaticMeshLODAssetTool.cpp b/Engine/Plugins/Experimental/MeshLODToolset/Source/MeshLODToolset/Private/Tools/GenerateStaticMeshLODAssetTool.cpp
--- a/Engine/Plugins/Experimental/MeshLODToolset/Source/MeshLODToolset/Private/Tools/GenerateStaticMeshLODAssetTool.cpp
+++ b/Engine/Plugins/Experimental/MeshLODToolset/Source/MeshLODToolset/Private/Tools/GenerateStaticMeshLODAssetTool.cpp
@@ -224,6 +224,7 @@ bool UGenerateStaticMeshLODAssetToolBuilder::CanBuildTool(const FToolBuilderStat
 UInteractiveTool* UGenerateStaticMeshLODAssetToolBuilder::BuildTool(const FToolBuilderState& SceneState) const
 {
 	UGenerateStaticMeshLODAssetTool* NewTool = NewObject<UGenerateStaticMeshLODAssetTool>(SceneState.ToolManager);
+	NewTool->SetUseAssetEditorMode(bUseAssetEditorMode);
 
 	TArray<TObjectPtr<UToolTarget>> Targets = SceneState.TargetManager->BuildAllSelectedTargetable(SceneState, GetTargetRequirements());
 	NewTool->SetTargets(MoveTemp(Targets));
@@ -255,6 +256,12 @@ void UGenerateStaticMeshLODAssetTool::SetWorld(UWorld* World)
 }
 
 
+void UGenerateStaticMeshLODAssetTool::SetUseAssetEditorMode(bool bEnable)
+{
+	bIsInAssetEditorMode = bEnable;
+}
+
+
 void UGenerateStaticMeshLODAssetTool::Setup()
 {
 	using GenerateStaticMeshLODAssetLocals::FGenerateStaticMeshLODAssetOperatorFactory;
@@ -310,6 +317,8 @@ void UGenerateStaticMeshLODAssetTool::Setup()
 	OutputProperties->NewAssetName = FPaths::GetBaseFilename(FullPathWithExtension, true);
 	OutputProperties->GeneratedSuffix = TEXT("_AutoLOD");
 	OutputProperties->RestoreProperties(this);
+	OutputProperties->OutputMode = (bIsInAssetEditorMode) ? EGenerateLODAssetOutputMode::UpdateExistingAsset : EGenerateLODAssetOutputMode::CreateNewAsset;
+	OutputProperties->bShowOutputMode = !bIsInAssetEditorMode;
 
 	PresetProperties = NewObject<UGenerateStaticMeshLODAssetToolPresetProperties>(this);
 	PresetProperties->RestoreProperties(this);
diff --git a/Engine/Plugins/Experimental/MeshLODToolset/Source/MeshLODToolset/Public/Tools/GenerateStaticMeshLODAssetTool.h b/Engine/Plugins/Experimental/MeshLODToolset/Source/MeshLODToolset/Public/Tools/GenerateStaticMeshLODAssetTool.h
--- a/Engine/Plugins/Experimental/MeshLODToolset/Source/MeshLODToolset/Public/Tools/GenerateStaticMeshLODAssetTool.h
+++ b/Engine/Plugins/Experimental/MeshLODToolset/Source/MeshLODToolset/Public/Tools/GenerateStaticMeshLODAssetTool.h
@@ -43,6 +43,8 @@ class MESHLODTOOLSET_API UGenerateStaticMeshLODAssetToolBuilder : public UIntera
 	GENERATED_BODY()
 
 public:
+	bool bUseAssetEditorMode = false;
+
 	virtual bool CanBuildTool(const FToolBuilderState& SceneState) const override;
 	virtual UInteractiveTool* BuildTool(const FToolBuilderState& SceneState) const override;
 
@@ -58,7 +60,7 @@ class MESHLODTOOLSET_API UGenerateStaticMeshLODAssetToolOutputProperties : publi
 	GENERATED_BODY()
 public:
 	/** Whether to modify the static mesh in place or create a new one. */
-	UPROPERTY(EditAnywhere, Category = "Output Options", meta = (TransientToolProperty))
+	UPROPERTY(EditAnywhere, Category = "Output Options", meta = (TransientToolProperty, EditConditionHides, EditCondition = "bShowOutputMode"))
 	EGenerateLODAssetOutputMode OutputMode = EGenerateLODAssetOutputMode::CreateNewAsset;
 
 	/** Base name for newly-generated asset */
@@ -72,6 +74,10 @@ public:
 	/** Suffix to append to newly-generated Asset (Meshes, Textures, Materials, etc) */
 	UPROPERTY(EditAnywhere, Category = "Output Options", meta = (TransientToolProperty))
 	FString GeneratedSuffix;
+
+	/** If false, then OutputMode will not be shown in DetailsView panels (otherwise no effect) */
+	UPROPERTY(meta = (TransientToolProperty))
+	bool bShowOutputMode = true;
 };
 
 
@@ -259,6 +265,10 @@ class MESHLODTOOLSET_API UGenerateStaticMeshLODAssetTool : public UMultiSelectio
 	GENERATED_BODY()
 
 public:
+
+	// Enable UI Customization for running this Tool in the Static Mesh Asset Editor. Must call before Setup.
+	virtual void SetUseAssetEditorMode(bool bEnable);
+
 	virtual void Setup() override;
 	virtual void Shutdown(EToolShutdownType ShutdownType) override;
 
@@ -317,6 +327,7 @@ protected:
 	TObjectPtr<UPreviewGeometry> CollisionPreview;
 
 protected:
+	bool bIsInAssetEditorMode = false;
 	UWorld* TargetWorld;
 
 	UPROPERTY()
@@ -335,5 +346,4 @@ protected:
 
 	void OnPresetSelectionChanged();
 
-
 };
diff --git a/Engine/Plugins/Experimental/StaticMeshEditorModeling/Source/StaticMeshEditorModeling/Private/StaticMeshEditorModelingMode.cpp b/Engine/Plugins/Experimental/StaticMeshEditorModeling/Source/StaticMeshEditorModeling/Private/StaticMeshEditorModelingMode.cpp
--- a/Engine/Plugins/Experimental/StaticMeshEditorModeling/Source/StaticMeshEditorModeling/Private/StaticMeshEditorModelingMode.cpp
+++ b/Engine/Plugins/Experimental/StaticMeshEditorModeling/Source/StaticMeshEditorModeling/Private/StaticMeshEditorModelingMode.cpp
@@ -27,7 +27,9 @@ void UStaticMeshEditorModelingMode::Enter()
 
 	const FModelingToolsManagerCommands& ToolManagerCommands = FModelingToolsManagerCommands::Get();
 
-	RegisterTool(ToolManagerCommands.BeginGenerateStaticMeshLODAssetTool, TEXT("BeginGenerateStaticMeshLODAssetTool"), NewObject<UGenerateStaticMeshLODAssetToolBuilder>());
+	UGenerateStaticMeshLODAssetToolBuilder* GenerateStaticMeshLODAssetToolBuilder = NewObject<UGenerateStaticMeshLODAssetToolBuilder>();
+	GenerateStaticMeshLODAssetToolBuilder->bUseAssetEditorMode = true;
+	RegisterTool(ToolManagerCommands.BeginGenerateStaticMeshLODAssetTool, TEXT("BeginGenerateStaticMeshLODAssetTool"), GenerateStaticMeshLODAssetToolBuilder);
 	RegisterTool(ToolManagerCommands.BeginLODManagerTool, TEXT("BeginLODManagerTool"), NewObject<ULODManagerToolBuilder>());
 	RegisterTool(ToolManagerCommands.BeginMeshInspectorTool, TEXT("BeginMeshInspectorTool"), NewObject<UMeshInspectorToolBuilder>());
 }
