AutoLOD: use HideEditConditionToggle to remove the checkbox next to Output Mode

#rnx

diff --git a/Engine/Plugins/Experimental/MeshLODToolset/Source/MeshLODToolset/Public/Tools/GenerateStaticMeshLODAssetTool.h b/Engine/Plugins/Experimental/MeshLODToolset/Source/MeshLODToolset/Public/Tools/GenerateStaticMeshLODAssetTool.h
--- a/Engine/Plugins/Experimental/MeshLODToolset/Source/MeshLODToolset/Public/Tools/GenerateStaticMeshLODAssetTool.h
+++ b/Engine/Plugins/Experimental/MeshLODToolset/Source/MeshLODToolset/Public/Tools/GenerateStaticMeshLODAssetTool.h
@@ -60,7 +60,7 @@ class MESHLODTOOLSET_API UGenerateStaticMeshLODAssetToolOutputProperties : publi
 	GENERATED_BODY()
 public:
 	/** Whether to modify the static mesh in place or create a new one. */
-	UPROPERTY(EditAnywhere, Category = "Output Options", meta = (TransientToolProperty, EditConditionHides, EditCondition = "bShowOutputMode"))
+	UPROPERTY(EditAnywhere, Category = "Output Options", meta = (TransientToolProperty, HideEditConditionToggle, EditCondition = "bShowOutputMode"))
 	EGenerateLODAssetOutputMode OutputMode = EGenerateLODAssetOutputMode::CreateNewAsset;
 
 	/** Base name for newly-generated asset */
