If a ToolTarget's MeshComponent has no MeshDescription, return an empty but valid one. Harden various tools to not crash on empty meshes.

#rnx

diff --git a/Engine/Plugins/Experimental/MeshLODToolset/Source/MeshLODToolset/Private/Graphs/GenerateStaticMeshLODProcess.cpp b/Engine/Plugins/Experimental/MeshLODToolset/Source/MeshLODToolset/Private/Graphs/GenerateStaticMeshLODProcess.cpp
--- a/Engine/Plugins/Experimental/MeshLODToolset/Source/MeshLODToolset/Private/Graphs/GenerateStaticMeshLODProcess.cpp
+++ b/Engine/Plugins/Experimental/MeshLODToolset/Source/MeshLODToolset/Private/Graphs/GenerateStaticMeshLODProcess.cpp
@@ -172,6 +172,12 @@ bool UGenerateStaticMeshLODProcess::Initialize(UStaticMesh* StaticMeshIn, FProgr
 	bUsingHiResSource = SourceStaticMesh->IsHiResMeshDescriptionValid();
 	const FMeshDescription* UseSourceMeshDescription =
 		(bUsingHiResSource) ? SourceStaticMesh->GetHiResMeshDescription() : SourceStaticMesh->GetMeshDescription(0);
+
+	if (!UseSourceMeshDescription)
+	{
+		return false;
+	}
+
 	SourceMeshDescription = MakeShared<FMeshDescription>();
 	*SourceMeshDescription = *UseSourceMeshDescription;
 
@@ -929,6 +935,11 @@ void UGenerateStaticMeshLODProcess::UpdateSourceBakeTextureConstraint(UTexture2D
 
 bool UGenerateStaticMeshLODProcess::ComputeDerivedSourceData(FProgressCancel* Progress)
 {
+	if (Generator == nullptr)
+	{
+		return false;
+	}
+
 	DerivedTextureImages.Reset();
 
 	Generator->EvaluateResult(
diff --git a/Engine/Plugins/Experimental/MeshLODToolset/Source/MeshLODToolset/Private/Tools/GenerateStaticMeshLODAssetTool.cpp b/Engine/Plugins/Experimental/MeshLODToolset/Source/MeshLODToolset/Private/Tools/GenerateStaticMeshLODAssetTool.cpp
--- a/Engine/Plugins/Experimental/MeshLODToolset/Source/MeshLODToolset/Private/Tools/GenerateStaticMeshLODAssetTool.cpp
+++ b/Engine/Plugins/Experimental/MeshLODToolset/Source/MeshLODToolset/Private/Tools/GenerateStaticMeshLODAssetTool.cpp
@@ -291,12 +291,23 @@ void UGenerateStaticMeshLODAssetTool::Setup()
 			const FProgressCancel::FMessageInfo& Warning = Progress.Warnings[0];
 			GetToolManager()->DisplayMessage(Warning.MessageText, (EToolMessageLevel)Warning.MessageLevel);
 		}
+
+		if (!bInitializeOK)
+		{
+			GenerateProcess = nullptr;
+			GetToolManager()->DisplayMessage(
+				LOCTEXT("GenerateStaticMeshLODAssetTool_ErrorInitializing", "Error initializing tool process: invalid Static Mesh input"),
+				EToolMessageLevel::UserError);
+			return;
+		}
+
 	}
 	else
 	{
 		GetToolManager()->DisplayMessage(
 			LOCTEXT("GenerateStaticMeshLODAssetTool_NoStaticMesh", "Could not find Static Mesh in selected input"),
 			EToolMessageLevel::UserError);
+		return;
 	}
 
 	FString FullPathWithExtension = UEditorAssetLibrary::GetPathNameForLoadedAsset(StaticMesh);
@@ -436,15 +447,29 @@ void UGenerateStaticMeshLODAssetTool::OnSettingsModified()
 
 void UGenerateStaticMeshLODAssetTool::Shutdown(EToolShutdownType ShutdownType)
 {
-	OutputProperties->SaveProperties(this);
-	BasicProperties->SaveProperties(this);
-	CollisionVizSettings->SaveProperties(this);
-	PresetProperties->SaveProperties(this);
-
-	CollisionPreview->Disconnect();
-	CollisionPreview = nullptr;
+	if (OutputProperties)
+	{
+		OutputProperties->SaveProperties(this);
+	}
+	if (BasicProperties)
+	{
+		BasicProperties->SaveProperties(this);
+	}
+	if (CollisionVizSettings)
+	{
+		CollisionVizSettings->SaveProperties(this);
+	}
+	if (PresetProperties)
+	{
+		PresetProperties->SaveProperties(this);
+	}
+	if (CollisionPreview)
+	{
+		CollisionPreview->Disconnect();
+		CollisionPreview = nullptr;
+	}
 
-	if (ShutdownType == EToolShutdownType::Accept)
+	if (ShutdownType == EToolShutdownType::Accept && GenerateProcess)
 	{
 		if (OutputProperties->OutputMode == EGenerateLODAssetOutputMode::UpdateExistingAsset)
 		{
@@ -456,7 +481,10 @@ void UGenerateStaticMeshLODAssetTool::Shutdown(EToolShutdownType ShutdownType)
 		}
 	}
 
-	FDynamicMeshOpResult Result = PreviewWithBackgroundCompute->Shutdown();
+	if (PreviewWithBackgroundCompute)
+	{
+		PreviewWithBackgroundCompute->Shutdown();
+	}
 }
 
 
diff --git a/Engine/Plugins/Experimental/MeshLODToolset/Source/MeshLODToolset/Public/Graphs/GenerateStaticMeshLODProcess.h b/Engine/Plugins/Experimental/MeshLODToolset/Source/MeshLODToolset/Public/Graphs/GenerateStaticMeshLODProcess.h
--- a/Engine/Plugins/Experimental/MeshLODToolset/Source/MeshLODToolset/Public/Graphs/GenerateStaticMeshLODProcess.h
+++ b/Engine/Plugins/Experimental/MeshLODToolset/Source/MeshLODToolset/Public/Graphs/GenerateStaticMeshLODProcess.h
@@ -347,7 +347,7 @@ public:
 	 *  - Read all the Source Texture Images that need to be baked
 	 *  - call InitializeGenerator() to set up the compute graph
 	 */
-	bool Initialize(UStaticMesh* SourceMesh, FProgressCancel* Progress = nullptr);
+	UE_NODISCARD bool Initialize(UStaticMesh* SourceMesh, FProgressCancel* Progress = nullptr);
 
 
 	/** @return the source UStaticMesh passed to Initialize() */
diff --git a/Engine/Plugins/Experimental/MeshModelingToolsetExp/Source/MeshModelingToolsEditorOnly/Private/SubdividePolyTool.cpp b/Engine/Plugins/Experimental/MeshModelingToolsetExp/Source/MeshModelingToolsEditorOnly/Private/SubdividePolyTool.cpp
--- a/Engine/Plugins/Experimental/MeshModelingToolsetExp/Source/MeshModelingToolsEditorOnly/Private/SubdividePolyTool.cpp
+++ b/Engine/Plugins/Experimental/MeshModelingToolsetExp/Source/MeshModelingToolsEditorOnly/Private/SubdividePolyTool.cpp
@@ -44,6 +44,11 @@ public:
 
 	void ProcessMesh(const FDynamicMesh3& Mesh, FDynamicMesh3& OutRenderMesh) final
 	{
+		if (Mesh.TriangleCount() == 0 || Mesh.VertexCount() == 0)
+		{
+			return;
+		}
+
 		constexpr bool bAutoCompute = true;
 		FGroupTopology Topo(&Mesh, bAutoCompute);
 		FSubdividePoly Subd(Topo, Mesh, SubdivisionLevel);
@@ -128,7 +133,7 @@ void USubdividePolyTool::CapSubdivisionLevel(ESubdivisionScheme Scheme, int Desi
 		FGroupTopology Topo(OriginalMesh.Get(), bAutoCompute);
 		NumOriginalFaces = Topo.Groups.Num();
 	}
-	int MaxLevel = (int)floor(log2(MaxFaces / NumOriginalFaces) / 2.0);
+	int MaxLevel = (int)floor(log2(MaxFaces / (NumOriginalFaces+1)) / 2.0);
 
 	if (DesiredLevel > MaxLevel)
 	{
diff --git a/Engine/Plugins/Experimental/MeshModelingToolsetExp/Source/MeshModelingToolsExp/Private/MeshSpaceDeformerTool.cpp b/Engine/Plugins/Experimental/MeshModelingToolsetExp/Source/MeshModelingToolsExp/Private/MeshSpaceDeformerTool.cpp
--- a/Engine/Plugins/Experimental/MeshModelingToolsetExp/Source/MeshModelingToolsExp/Private/MeshSpaceDeformerTool.cpp
+++ b/Engine/Plugins/Experimental/MeshModelingToolsetExp/Source/MeshModelingToolsExp/Private/MeshSpaceDeformerTool.cpp
@@ -162,8 +162,8 @@ void UMeshSpaceDeformerTool::Setup()
 	// with the Z axis (along which the deformation acts) aligned to the longest of the
 	// bounding box dimensions, after scaling them with the transform.
 	FAxisAlignedBox3d BBox = OriginalDynamicMesh->GetBounds();
-	FVector3d Dimensions = BBox.Max - BBox.Min;
-	MeshCenter = (FVector3d)MeshTransform.TransformPosition((FVector)BBox.Center());
+	FVector3d Dimensions = BBox.IsEmpty() ? FVector3d::Zero() : BBox.Max - BBox.Min;
+	MeshCenter = BBox.IsEmpty() ? FVector3d::Zero() : (FVector3d)MeshTransform.TransformPosition((FVector)BBox.Center());
 
 	Dimensions = FVector3d(MeshTransform.GetScale3D().GetAbs()) * Dimensions;
 	double WorldMajorLength = 0;
diff --git a/Engine/Plugins/Experimental/MeshModelingToolsetExp/Source/MeshModelingToolsExp/Private/WeldMeshEdgesTool.cpp b/Engine/Plugins/Experimental/MeshModelingToolsetExp/Source/MeshModelingToolsExp/Private/WeldMeshEdgesTool.cpp
--- a/Engine/Plugins/Experimental/MeshModelingToolsetExp/Source/MeshModelingToolsExp/Private/WeldMeshEdgesTool.cpp
+++ b/Engine/Plugins/Experimental/MeshModelingToolsetExp/Source/MeshModelingToolsExp/Private/WeldMeshEdgesTool.cpp
@@ -46,6 +46,11 @@ public:
 			return;
 		}
 
+		if (SourceMesh->TriangleCount() == 0 || SourceMesh->VertexCount() == 0)
+		{
+			return;
+		}
+
 		ResultMesh->Copy(*SourceMesh, true, true, true, true);
 	
 		FMergeCoincidentMeshEdges Merger(ResultMesh.Get());
diff --git a/Engine/Plugins/Runtime/GeometryProcessing/Source/DynamicMesh/Private/Operations/FFDLattice.cpp b/Engine/Plugins/Runtime/GeometryProcessing/Source/DynamicMesh/Private/Operations/FFDLattice.cpp
--- a/Engine/Plugins/Runtime/GeometryProcessing/Source/DynamicMesh/Private/Operations/FFDLattice.cpp
+++ b/Engine/Plugins/Runtime/GeometryProcessing/Source/DynamicMesh/Private/Operations/FFDLattice.cpp
@@ -13,7 +13,10 @@ FFFDLattice::FFFDLattice(const FVector3i& InDims, const FDynamicMesh3& Mesh, flo
 	check(InDims.X > 1 && InDims.Y > 1 && InDims.Z > 1);
 
 	InitialBounds = Mesh.GetBounds();
-	check(!InitialBounds.IsEmpty());
+	if (InitialBounds.IsEmpty())
+	{
+		InitialBounds = FAxisAlignedBox3d(FVector3d::Zero(), FVector3d::Zero());
+	}
 
 	// Expand the initial bounding box to make the computation of which grid cell a mesh vertex is inside of a little 
 	// less susceptible to numerical error issues.
diff --git a/Engine/Plugins/Runtime/GeometryProcessing/Source/DynamicMesh/Private/ShapeApproximation/MeshSimpleShapeApproximation.cpp b/Engine/Plugins/Runtime/GeometryProcessing/Source/DynamicMesh/Private/ShapeApproximation/MeshSimpleShapeApproximation.cpp
--- a/Engine/Plugins/Runtime/GeometryProcessing/Source/DynamicMesh/Private/ShapeApproximation/MeshSimpleShapeApproximation.cpp
+++ b/Engine/Plugins/Runtime/GeometryProcessing/Source/DynamicMesh/Private/ShapeApproximation/MeshSimpleShapeApproximation.cpp
@@ -187,12 +187,15 @@ void FMeshSimpleShapeApproximation::Generate_AlignedBoxes(FSimpleShapeSet3d& Sha
 
 		FAxisAlignedBox3d Bounds = SourceMeshes[idx]->GetBounds();
 
-		FBoxShape3d NewBox;
-		NewBox.Box = FOrientedBox3d(Bounds);
+		if (!Bounds.IsEmpty())
+		{
+			FBoxShape3d NewBox;
+			NewBox.Box = FOrientedBox3d(Bounds);
 
-		GeometryLock.Lock();
-		ShapeSetOut.Boxes.Add(NewBox);
-		GeometryLock.Unlock();
+			GeometryLock.Lock();
+			ShapeSetOut.Boxes.Add(NewBox);
+			GeometryLock.Unlock();
+		}
 	});
 }
 
diff --git a/Engine/Plugins/Runtime/MeshModelingToolset/Source/ModelingComponents/Private/Drawing/MeshElementsVisualizer.cpp b/Engine/Plugins/Runtime/MeshModelingToolset/Source/ModelingComponents/Private/Drawing/MeshElementsVisualizer.cpp
--- a/Engine/Plugins/Runtime/MeshModelingToolset/Source/ModelingComponents/Private/Drawing/MeshElementsVisualizer.cpp
+++ b/Engine/Plugins/Runtime/MeshModelingToolset/Source/ModelingComponents/Private/Drawing/MeshElementsVisualizer.cpp
@@ -24,7 +24,15 @@ public:
 	virtual FBoxSphereBounds GetBounds() const 
 	{ 
 		FAxisAlignedBox3d Bounds = Mesh.GetBounds();
-		return FBoxSphereBounds((FBox)Bounds);
+
+		if (Bounds.IsEmpty())
+		{
+			return FBoxSphereBounds(EForceInit::ForceInit);
+		}
+		else
+		{
+			return FBoxSphereBounds((FBox)Bounds);
+		}
 	}
 
 	virtual FVector GetVertex(int32 Index) const
diff --git a/Engine/Plugins/Runtime/MeshModelingToolset/Source/ModelingComponentsEditorOnly/Private/ToolTargets/StaticMeshComponentToolTarget.cpp b/Engine/Plugins/Runtime/MeshModelingToolset/Source/ModelingComponentsEditorOnly/Private/ToolTargets/StaticMeshComponentToolTarget.cpp
--- a/Engine/Plugins/Runtime/MeshModelingToolset/Source/ModelingComponentsEditorOnly/Private/ToolTargets/StaticMeshComponentToolTarget.cpp
+++ b/Engine/Plugins/Runtime/MeshModelingToolset/Source/ModelingComponentsEditorOnly/Private/ToolTargets/StaticMeshComponentToolTarget.cpp
@@ -10,6 +10,7 @@
 #include "Materials/Material.h"
 #include "ToolTargets/StaticMeshToolTarget.h"
 #include "AssetUtils/MeshDescriptionUtil.h"
+#include "StaticMeshAttributes.h"
 
 using namespace UE::Geometry;
 
@@ -126,11 +127,27 @@ bool UStaticMeshComponentToolTarget::CommitMaterialSetUpdate(const FComponentMat
 
 const FMeshDescription* UStaticMeshComponentToolTarget::GetMeshDescription()
 {
+	static FMeshDescription EmptyMeshDescription;
+	static bool bFirst = true;
+	if (bFirst)
+	{
+		FStaticMeshAttributes Attributes(EmptyMeshDescription);
+		Attributes.Register();
+		bFirst = false;
+	}
+
 	if (ensure(IsValid()))
 	{
 		UStaticMesh* StaticMesh = Cast<UStaticMeshComponent>(Component)->GetStaticMesh();
-		return (EditingLOD == EStaticMeshEditingLOD::HiResSource) ?
+		FMeshDescription* MeshDescription = (EditingLOD == EStaticMeshEditingLOD::HiResSource) ?
 			StaticMesh->GetHiResMeshDescription() : StaticMesh->GetMeshDescription((int32)EditingLOD);
+
+		if (!MeshDescription)
+		{
+			return &EmptyMeshDescription;
+		}
+
+		return MeshDescription;
 	}
 	return nullptr;
 }
diff --git a/Engine/Plugins/Runtime/MeshModelingToolset/Source/ModelingComponentsEditorOnly/Private/ToolTargets/StaticMeshToolTarget.cpp b/Engine/Plugins/Runtime/MeshModelingToolset/Source/ModelingComponentsEditorOnly/Private/ToolTargets/StaticMeshToolTarget.cpp
--- a/Engine/Plugins/Runtime/MeshModelingToolset/Source/ModelingComponentsEditorOnly/Private/ToolTargets/StaticMeshToolTarget.cpp
+++ b/Engine/Plugins/Runtime/MeshModelingToolset/Source/ModelingComponentsEditorOnly/Private/ToolTargets/StaticMeshToolTarget.cpp
@@ -10,6 +10,7 @@
 #include "RenderingThread.h"
 #include "Widgets/Notifications/SNotificationList.h"
 #include "Framework/Notifications/NotificationManager.h"
+#include "StaticMeshAttributes.h"
 
 using namespace UE::Geometry;
 
@@ -185,10 +186,26 @@ bool UStaticMeshToolTarget::CommitMaterialSetUpdate(UStaticMesh* StaticMeshIn,
 
 const FMeshDescription* UStaticMeshToolTarget::GetMeshDescription()
 {
+	static FMeshDescription EmptyMeshDescription;
+	static bool bFirst = true;
+	if (bFirst)
+	{
+		FStaticMeshAttributes Attributes(EmptyMeshDescription);
+		Attributes.Register();
+		bFirst = false;
+	}
+
 	if (ensure(IsValid()))
 	{
-		return (EditingLOD == EStaticMeshEditingLOD::HiResSource) ?
+		FMeshDescription* MeshDescription = (EditingLOD == EStaticMeshEditingLOD::HiResSource) ?
 			StaticMesh->GetHiResMeshDescription() : StaticMesh->GetMeshDescription((int32)EditingLOD);
+
+		if (!MeshDescription)
+		{
+			return &EmptyMeshDescription;
+		}
+
+		return MeshDescription;
 	}
 	return nullptr;
 }
diff --git a/Engine/Plugins/Runtime/MeshModelingToolset/Source/ModelingOperatorsEditorOnly/Private/CuttingOps/EmbedPolygonsOp.cpp b/Engine/Plugins/Runtime/MeshModelingToolset/Source/ModelingOperatorsEditorOnly/Private/CuttingOps/EmbedPolygonsOp.cpp
--- a/Engine/Plugins/Runtime/MeshModelingToolset/Source/ModelingOperatorsEditorOnly/Private/CuttingOps/EmbedPolygonsOp.cpp
+++ b/Engine/Plugins/Runtime/MeshModelingToolset/Source/ModelingOperatorsEditorOnly/Private/CuttingOps/EmbedPolygonsOp.cpp
@@ -112,6 +112,11 @@ void CollapseDegenerateEdgesOnVertexPath(FDynamicMesh3& Mesh, TArray<int>& Verte
 
 void FEmbedPolygonsOp::BooleanPath(FProgressCancel* Progress)
 {
+	if (ResultMesh->TriangleCount() == 0 || ResultMesh->VertexCount() == 0)
+	{
+		return;
+	}
+
 	FAxisAlignedBox3d Bounds = OriginalMesh->GetBounds();
 	Bounds.Expand(.01); // expand a little beyond bounds to avoid creating coplanar cases
 	double MeshDiameter = Bounds.MaxDim();
diff --git a/Engine/Source/Runtime/GeometryCore/Public/Image/ImageTile.h b/Engine/Source/Runtime/GeometryCore/Public/Image/ImageTile.h
--- a/Engine/Source/Runtime/GeometryCore/Public/Image/ImageTile.h
+++ b/Engine/Source/Runtime/GeometryCore/Public/Image/ImageTile.h
@@ -120,8 +120,8 @@ public:
 	{
 		checkSlow(TileWidthIn >= 1 && TileHeightIn >= 1);
 		Dimensions = DimensionsIn;
-		TileWidth = FMath::Clamp(TileWidthIn, 1, Dimensions.GetWidth());
-		TileHeight = FMath::Clamp(TileHeightIn, 1, Dimensions.GetHeight());
+		TileWidth = FMath::Clamp(TileWidthIn, 1, FMath::Max(1, Dimensions.GetWidth()));
+		TileHeight = FMath::Clamp(TileHeightIn, 1, FMath::Max(1, Dimensions.GetHeight()));
 	}
 
 	int32 NumTilesX() const
