ModelingTools: Fix issue where AutoLOD tool allows invalid grid resolutions when using the Convex Hull algorithm for generation.

#rnx

diff --git a/Engine/Plugins/Experimental/MeshLODToolset/Source/MeshLODToolset/Public/Graphs/GenerateStaticMeshLODProcess.h b/Engine/Plugins/Experimental/MeshLODToolset/Source/MeshLODToolset/Public/Graphs/GenerateStaticMeshLODProcess.h
--- a/Engine/Plugins/Experimental/MeshLODToolset/Source/MeshLODToolset/Public/Graphs/GenerateStaticMeshLODProcess.h
+++ b/Engine/Plugins/Experimental/MeshLODToolset/Source/MeshLODToolset/Public/Graphs/GenerateStaticMeshLODProcess.h
@@ -67,7 +67,7 @@ struct FGenerateStaticMeshLODProcessSettings
 	bool bPrefilterVertices = true;
 
 	/** Grid resolution (along the maximum-length axis) for subsampling before computing the convex hull */
-	UPROPERTY(EditAnywhere, Category = Collision, meta = (NoSpinbox = "true", EditConditionHides, EditCondition = "MeshGenerator == EGenerateStaticMeshLODProcess_MeshGeneratorModes::ConvexHull", UIMin = 4, UIMax = 100))
+	UPROPERTY(EditAnywhere, Category = Collision, meta = (NoSpinbox = "true", EditConditionHides, EditCondition = "MeshGenerator == EGenerateStaticMeshLODProcess_MeshGeneratorModes::ConvexHull", ClampMin = "4", ClampMax = "100"))
 	int PrefilterGridResolution = 10;
 
 };
