ModelingTools: disable multiple input objects for LODManagerTool. Only show panel for maniuplating HiRes SourceModel if the input mesh has one.
#rnx

diff --git a/Engine/Plugins/Experimental/MeshLODToolset/Source/MeshLODToolset/Private/Tools/LODManagerTool.cpp b/Engine/Plugins/Experimental/MeshLODToolset/Source/MeshLODToolset/Private/Tools/LODManagerTool.cpp
--- a/Engine/Plugins/Experimental/MeshLODToolset/Source/MeshLODToolset/Private/Tools/LODManagerTool.cpp
+++ b/Engine/Plugins/Experimental/MeshLODToolset/Source/MeshLODToolset/Private/Tools/LODManagerTool.cpp
@@ -45,6 +45,12 @@ const FToolTargetTypeRequirements& ULODManagerToolBuilder::GetTargetRequirements
 	return TypeRequirements;
 }
 
+bool ULODManagerToolBuilder::CanBuildTool(const FToolBuilderState& SceneState) const
+{
+	// disable multi-selection for now
+	return SceneState.TargetManager->CountSelectedAndTargetable(SceneState, GetTargetRequirements()) == 1;
+}
+
 UMultiSelectionMeshEditingTool* ULODManagerToolBuilder::CreateNewTool(const FToolBuilderState& SceneState) const
 {
 	return NewObject<ULODManagerTool>(SceneState.ToolManager);
@@ -73,45 +79,51 @@ void ULODManagerTool::Setup()
 {
 	UInteractiveTool::Setup();
 
-
-	if (Targets.Num() == 1)
+	if (! ensure(Targets.Num() == 1) )
 	{
-		LODPreviewProperties = NewObject<ULODManagerPreviewLODProperties>(this);
-		AddToolPropertySource(LODPreviewProperties);
-		LODPreviewProperties->VisibleLOD = this->DefaultLODName;
-		LODPreviewProperties->WatchProperty(LODPreviewProperties->VisibleLOD, [this](FString NewLOD) { bPreviewLODValid = false; });
-		LODPreviewProperties->WatchProperty(LODPreviewProperties->bShowSeams, [this](bool bNewValue) { bPreviewLODValid = false; });
+		return;
+	}
 
-		LODPreview = NewObject<UPreviewMesh>(this);
-		LODPreview->CreateInWorld(TargetWorld, (FTransform)UE::ToolTarget::GetLocalToWorldTransform(Targets[0]));
-		LODPreview->SetTangentsMode(EDynamicMeshComponentTangentsMode::ExternallyProvided);
-		LODPreview->SetVisible(false);
 
-		LODPreviewLines = NewObject<UPreviewGeometry>(this);
-		LODPreviewLines->CreateInWorld(TargetWorld, (FTransform)UE::ToolTarget::GetLocalToWorldTransform(Targets[0]));
+	LODPreviewProperties = NewObject<ULODManagerPreviewLODProperties>(this);
+	AddToolPropertySource(LODPreviewProperties);
+	LODPreviewProperties->VisibleLOD = this->DefaultLODName;
+	LODPreviewProperties->WatchProperty(LODPreviewProperties->VisibleLOD, [this](FString NewLOD) { bPreviewLODValid = false; });
+	LODPreviewProperties->WatchProperty(LODPreviewProperties->bShowSeams, [this](bool bNewValue) { bPreviewLODValid = false; });
 
-		FComponentMaterialSet MaterialSet = UE::ToolTarget::GetMaterialSet(Targets[0]);
-		LODPreview->SetMaterials(MaterialSet.Materials);
+	LODPreview = NewObject<UPreviewMesh>(this);
+	LODPreview->CreateInWorld(TargetWorld, (FTransform)UE::ToolTarget::GetLocalToWorldTransform(Targets[0]));
+	LODPreview->SetTangentsMode(EDynamicMeshComponentTangentsMode::ExternallyProvided);
+	LODPreview->SetVisible(false);
 
-		bLODInfoValid = false;
-	}
+	LODPreviewLines = NewObject<UPreviewGeometry>(this);
+	LODPreviewLines->CreateInWorld(TargetWorld, (FTransform)UE::ToolTarget::GetLocalToWorldTransform(Targets[0]));
 
+	FComponentMaterialSet MaterialSet = UE::ToolTarget::GetMaterialSet(Targets[0]);
+	LODPreview->SetMaterials(MaterialSet.Materials);
+
+	bLODInfoValid = false;
 
-	HiResSourceModelActions = NewObject<ULODManagerHiResSourceModelActions>(this);
-	HiResSourceModelActions->Initialize(this);
-	AddToolPropertySource(HiResSourceModelActions);
+	if (UStaticMeshComponent* StaticMeshComponent = Cast<UStaticMeshComponent>(UE::ToolTarget::GetTargetComponent(Targets[0])))
+	{
+		if (UStaticMesh* StaticMesh = StaticMeshComponent->GetStaticMesh())
+		{
+			if (StaticMesh->IsHiResMeshDescriptionValid())
+			{
+				HiResSourceModelActions = NewObject<ULODManagerHiResSourceModelActions>(this);
+				HiResSourceModelActions->Initialize(this);
+				AddToolPropertySource(HiResSourceModelActions);
+			}
+		}
+	}
 
 	MaterialActions = NewObject<ULODManagerMaterialActions>(this);
 	MaterialActions->Initialize(this);
 	AddToolPropertySource(MaterialActions);
 
-
-	if (Targets.Num() == 1)
-	{
-		LODInfoProperties = NewObject<ULODManagerLODProperties>(this);
-		AddToolPropertySource(LODInfoProperties);
-		bLODInfoValid = false;
-	}
+	LODInfoProperties = NewObject<ULODManagerLODProperties>(this);
+	AddToolPropertySource(LODInfoProperties);
+	bLODInfoValid = false;
 
 	SetToolDisplayName(LOCTEXT("ToolName", "Manage LODs"));
 	GetToolManager()->DisplayMessage(
diff --git a/Engine/Plugins/Experimental/MeshLODToolset/Source/MeshLODToolset/Public/Tools/LODManagerTool.h b/Engine/Plugins/Experimental/MeshLODToolset/Source/MeshLODToolset/Public/Tools/LODManagerTool.h
--- a/Engine/Plugins/Experimental/MeshLODToolset/Source/MeshLODToolset/Public/Tools/LODManagerTool.h
+++ b/Engine/Plugins/Experimental/MeshLODToolset/Source/MeshLODToolset/Public/Tools/LODManagerTool.h
@@ -24,6 +24,7 @@ class MESHLODTOOLSET_API ULODManagerToolBuilder : public UMultiSelectionMeshEdit
 	GENERATED_BODY()
 
 public:
+	virtual bool CanBuildTool(const FToolBuilderState& SceneState) const override;
 	virtual UMultiSelectionMeshEditingTool* CreateNewTool(const FToolBuilderState& SceneState) const override;
 
 protected:
