MeshModelingToolset: Fix Linux/Mac issues

diff --git a/Engine/Plugins/Editor/ModelingToolsEditorMode/ModelingToolsEditorMode.uplugin b/Engine/Plugins/Editor/ModelingToolsEditorMode/ModelingToolsEditorMode.uplugin
--- a/Engine/Plugins/Editor/ModelingToolsEditorMode/ModelingToolsEditorMode.uplugin
+++ b/Engine/Plugins/Editor/ModelingToolsEditorMode/ModelingToolsEditorMode.uplugin
@@ -17,8 +17,7 @@
 		{
 			"Name": "ModelingToolsEditorMode",
 			"Type": "Editor",
-			"LoadingPhase": "Default",
-			"PlatformAllowList": [ "Win64" ]
+			"LoadingPhase": "Default"
 		}
 	],
 	"Plugins": [
diff --git a/Engine/Plugins/Experimental/GeometryFlow/GeometryFlow.uplugin b/Engine/Plugins/Experimental/GeometryFlow/GeometryFlow.uplugin
--- a/Engine/Plugins/Experimental/GeometryFlow/GeometryFlow.uplugin
+++ b/Engine/Plugins/Experimental/GeometryFlow/GeometryFlow.uplugin
@@ -28,9 +28,8 @@
 		{
 			"Name": "GeometryFlowMeshProcessingEditor",
 			"Type": "Editor",
-			"LoadingPhase": "Default",
-			"PlatformAllowList": [ "Win64" ]
-		},		
+			"LoadingPhase": "Default"
+		}
 	],
 	"Plugins": [
 		{
diff --git a/Engine/Plugins/Experimental/MeshLODToolset/MeshLODToolset.uplugin b/Engine/Plugins/Experimental/MeshLODToolset/MeshLODToolset.uplugin
--- a/Engine/Plugins/Experimental/MeshLODToolset/MeshLODToolset.uplugin
+++ b/Engine/Plugins/Experimental/MeshLODToolset/MeshLODToolset.uplugin
@@ -18,8 +18,7 @@
 		{
 			"Name": "MeshLODToolset",
 			"Type": "Editor",
-			"LoadingPhase": "Default",
-			"PlatformAllowList": [ "Win64" ]
+			"LoadingPhase": "Default"
 		},
 	],
 	"Plugins": [
diff --git a/Engine/Plugins/Experimental/MeshLODToolset/Source/MeshLODToolset/Private/Graphs/GenerateStaticMeshLODProcess.cpp b/Engine/Plugins/Experimental/MeshLODToolset/Source/MeshLODToolset/Private/Graphs/GenerateStaticMeshLODProcess.cpp
--- a/Engine/Plugins/Experimental/MeshLODToolset/Source/MeshLODToolset/Private/Graphs/GenerateStaticMeshLODProcess.cpp
+++ b/Engine/Plugins/Experimental/MeshLODToolset/Source/MeshLODToolset/Private/Graphs/GenerateStaticMeshLODProcess.cpp
@@ -23,7 +23,7 @@
 #include "AssetRegistryModule.h"
 #include "IAssetTools.h"
 #include "FileHelpers.h"
-#include "UObject/Metadata.h"
+#include "UObject/MetaData.h"
 
 #include "Engine/Engine.h"
 #include "Editor.h"
diff --git a/Engine/Plugins/Experimental/MeshModelingToolsetExp/MeshModelingToolsetExp.uplugin b/Engine/Plugins/Experimental/MeshModelingToolsetExp/MeshModelingToolsetExp.uplugin
--- a/Engine/Plugins/Experimental/MeshModelingToolsetExp/MeshModelingToolsetExp.uplugin
+++ b/Engine/Plugins/Experimental/MeshModelingToolsetExp/MeshModelingToolsetExp.uplugin
@@ -23,8 +23,7 @@
 		{
 			"Name": "MeshModelingToolsEditorOnly",
 			"Type": "Editor",
-			"LoadingPhase": "Default",
-			"PlatformAllowList": [ "Win64" ]
+			"LoadingPhase": "Default"
 		},
 		{
 			"Name": "GeometryProcessingAdapters",
diff --git a/Engine/Plugins/Runtime/MeshModelingToolset/MeshModelingToolset.uplugin b/Engine/Plugins/Runtime/MeshModelingToolset/MeshModelingToolset.uplugin
--- a/Engine/Plugins/Runtime/MeshModelingToolset/MeshModelingToolset.uplugin
+++ b/Engine/Plugins/Runtime/MeshModelingToolset/MeshModelingToolset.uplugin
@@ -28,8 +28,7 @@
 		{
 			"Name": "ModelingComponentsEditorOnly",
 			"Type": "Editor",
-			"LoadingPhase": "Default",
-			"PlatformAllowList": [ "Win64" ]
+			"LoadingPhase": "Default"
 		},		
 		{
 			"Name": "ModelingOperators",
