Don't allow AutoLOD to modify built-in Editor assets

#rnx

diff --git a/Engine/Plugins/Experimental/MeshLODToolset/Source/MeshLODToolset/Private/Tools/GenerateStaticMeshLODAssetTool.cpp b/Engine/Plugins/Experimental/MeshLODToolset/Source/MeshLODToolset/Private/Tools/GenerateStaticMeshLODAssetTool.cpp
--- a/Engine/Plugins/Experimental/MeshLODToolset/Source/MeshLODToolset/Private/Tools/GenerateStaticMeshLODAssetTool.cpp
+++ b/Engine/Plugins/Experimental/MeshLODToolset/Source/MeshLODToolset/Private/Tools/GenerateStaticMeshLODAssetTool.cpp
@@ -197,6 +197,14 @@ namespace GenerateStaticMeshLODAssetLocals
 		UE::Geometry::FTransform3d ResultTransform;
 	};
 
+	static void DisplayCriticalWarningMessage(const FString& Message)
+	{
+		FNotificationInfo Info(FText::FromString(Message));
+		Info.ExpireDuration = 5.0f;
+		FSlateNotificationManager::Get().AddNotification(Info);
+
+		UE_LOG(LogTemp, Warning, TEXT("%s"), *Message);
+	}
 }
 
 
@@ -540,6 +548,15 @@ void UGenerateStaticMeshLODAssetTool::CreateNewAsset()
 void UGenerateStaticMeshLODAssetTool::UpdateExistingAsset()
 {
 	check(PreviewWithBackgroundCompute->HaveValidResult());
+
+	const UStaticMesh* SourceStaticMesh = GenerateProcess->GetSourceStaticMesh();
+	if (SourceStaticMesh->GetPathName().StartsWith(TEXT("/Engine/")))
+	{
+		GenerateStaticMeshLODAssetLocals::DisplayCriticalWarningMessage(FString::Printf(TEXT("CANNOT MODIFY BUILT-IN ENGINE ASSET %s"), *SourceStaticMesh->GetPathName()));
+		return;
+	}
+
+
 	GenerateProcess->UpdateDerivedPathName(OutputProperties->NewAssetName, OutputProperties->GeneratedSuffix);
 
 	check(GenerateProcess->GraphEvalCriticalSection.TryLock());		// No ops should be running
