AutoLOD: ensure all available parameters are copied up from the Graph object to the Process object on init, so that they can be reflected in the Tool UI immediately

#rnx

diff --git a/Engine/Plugins/Experimental/MeshLODToolset/Source/MeshLODToolset/Private/Graphs/GenerateMeshLODGraph.cpp b/Engine/Plugins/Experimental/MeshLODToolset/Source/MeshLODToolset/Private/Graphs/GenerateMeshLODGraph.cpp
--- a/Engine/Plugins/Experimental/MeshLODToolset/Source/MeshLODToolset/Private/Graphs/GenerateMeshLODGraph.cpp
+++ b/Engine/Plugins/Experimental/MeshLODToolset/Source/MeshLODToolset/Private/Graphs/GenerateMeshLODGraph.cpp
@@ -493,6 +493,10 @@ void FGenerateMeshLODGraph::BuildGraph(const FDynamicMesh3* SourceMeshHint)
 	PreFilterSettings.FilterGroupLayerName = FName(TEXT("PreFilterGroups"));
 	UpdatePreFilterSettings(PreFilterSettings);
 
+	UE::GeometryFlow::FMeshThickenSettings MeshThickenSettings;
+	MeshThickenSettings.ThickenAmount = 10.0f;
+	UpdateThickenSettings(MeshThickenSettings);
+
 	FMeshSolidifySettings SolidifySettings;
 	SolidifySettings.VoxelResolution = 128;
 	UpdateSolidifySettings(SolidifySettings);
diff --git a/Engine/Plugins/Experimental/MeshLODToolset/Source/MeshLODToolset/Private/Graphs/GenerateStaticMeshLODProcess.cpp b/Engine/Plugins/Experimental/MeshLODToolset/Source/MeshLODToolset/Private/Graphs/GenerateStaticMeshLODProcess.cpp
--- a/Engine/Plugins/Experimental/MeshLODToolset/Source/MeshLODToolset/Private/Graphs/GenerateStaticMeshLODProcess.cpp
+++ b/Engine/Plugins/Experimental/MeshLODToolset/Source/MeshLODToolset/Private/Graphs/GenerateStaticMeshLODProcess.cpp
@@ -597,15 +597,18 @@ bool UGenerateStaticMeshLODProcess::InitializeGenerator()
 	Generator->SetSourceMesh(this->SourceMesh);
 
 
-	// read back default settings
+	// read parameter settings from the graph
 
 	CurrentSettings_Preprocess.FilterGroupLayer = Generator->GetCurrentPreFilterSettings().FilterGroupLayerName;
+	CurrentSettings_Preprocess.ThickenAmount = Generator->GetCurrentThickenSettings().ThickenAmount;
+	CurrentSettings_Preprocess.ThickenWeightMapName = "";
+	CurrentSettings.MeshGenerator = static_cast<EGenerateStaticMeshLODProcess_MeshGeneratorModes>(Generator->GetCurrentCoreMeshGeneratorMode());
 
 	CurrentSettings.SolidifyVoxelResolution = Generator->GetCurrentSolidifySettings().VoxelResolution;
 	CurrentSettings.WindingThreshold = Generator->GetCurrentSolidifySettings().WindingThreshold;
-
-	//CurrentSettings.MorphologyVoxelResolution = Generator->GetCurrentMorphologySettings().VoxelResolution;
 	CurrentSettings.ClosureDistance = Generator->GetCurrentMorphologySettings().Distance;
+	CurrentSettings.bPrefilterVertices = Generator->GetCurrentGenerateConvexHullMeshSettings().bPrefilterVertices;
+	CurrentSettings.PrefilterGridResolution = Generator->GetCurrentGenerateConvexHullMeshSettings().PrefilterGridResolution;
 
 	CurrentSettings_Simplify.Method = static_cast<EGenerateStaticMeshLODProcess_SimplifyMethod>(Generator->GetCurrentSimplifySettings().TargetType);
 	CurrentSettings_Simplify.TargetCount = Generator->GetCurrentSimplifySettings().TargetCount;
