AutoLOD: replace spaces in file path with underscores

#rnx

diff --git a/Engine/Plugins/Experimental/MeshLODToolset/Source/MeshLODToolset/Private/Graphs/GenerateStaticMeshLODProcess.cpp b/Engine/Plugins/Experimental/MeshLODToolset/Source/MeshLODToolset/Private/Graphs/GenerateStaticMeshLODProcess.cpp
--- a/Engine/Plugins/Experimental/MeshLODToolset/Source/MeshLODToolset/Private/Graphs/GenerateStaticMeshLODProcess.cpp
+++ b/Engine/Plugins/Experimental/MeshLODToolset/Source/MeshLODToolset/Private/Graphs/GenerateStaticMeshLODProcess.cpp
@@ -458,13 +458,13 @@ bool UGenerateStaticMeshLODProcess::Initialize(UStaticMesh* StaticMeshIn, FProgr
 
 void UGenerateStaticMeshLODProcess::UpdateDerivedPathName(const FString& NewAssetBaseName, const FString& NewAssetSuffix)
 {
-	DerivedAssetNameNoSuffix = FPaths::MakeValidFileName(NewAssetBaseName);
+	DerivedAssetNameNoSuffix = FPaths::MakeValidFileName(NewAssetBaseName.Replace(TEXT(" "), TEXT("_")));
 	if (DerivedAssetNameNoSuffix.Len() == 0)
 	{
 		DerivedAssetNameNoSuffix = SourceAssetName;
 	}
 
-	DerivedSuffix = FPaths::MakeValidFileName(NewAssetSuffix);
+	DerivedSuffix = FPaths::MakeValidFileName(NewAssetSuffix.Replace(TEXT(" "), TEXT("_")));
 	if (DerivedSuffix.Len() == 0)
 	{
 		DerivedSuffix = GetDefaultDerivedAssetSuffix();
