AutoLOD: don't allow Accept if the output mesh is empty

#rnx

diff --git a/Engine/Plugins/Experimental/MeshLODToolset/Source/MeshLODToolset/Private/Tools/GenerateStaticMeshLODAssetTool.cpp b/Engine/Plugins/Experimental/MeshLODToolset/Source/MeshLODToolset/Private/Tools/GenerateStaticMeshLODAssetTool.cpp
--- a/Engine/Plugins/Experimental/MeshLODToolset/Source/MeshLODToolset/Private/Tools/GenerateStaticMeshLODAssetTool.cpp
+++ b/Engine/Plugins/Experimental/MeshLODToolset/Source/MeshLODToolset/Private/Tools/GenerateStaticMeshLODAssetTool.cpp
@@ -377,6 +377,9 @@ void UGenerateStaticMeshLODAssetTool::Setup()
 	ToolSetupUtil::ApplyRenderingConfigurationToPreview(PreviewWithBackgroundCompute->PreviewMesh, nullptr);
 	PreviewWithBackgroundCompute->PreviewMesh->SetTangentsMode(EDynamicMeshComponentTangentsMode::ExternallyProvided);
 
+	PreviewWithBackgroundCompute->OnMeshUpdated.AddLambda([this](UMeshOpPreviewWithBackgroundCompute* Compute) {
+		UpdateAcceptWarnings(Compute->HaveEmptyResult() ? EAcceptWarning::EmptyForbidden : EAcceptWarning::NoWarning);
+	});
 
 	// For the first computation, display a bounding box with the working material. Otherwise it looks like nothing
 	// is happening. And we don't want to copy over the potentially huge input mesh to be the preview mesh.
@@ -498,7 +501,7 @@ void UGenerateStaticMeshLODAssetTool::Shutdown(EToolShutdownType ShutdownType)
 
 bool UGenerateStaticMeshLODAssetTool::CanAccept() const
 {
-	return (PreviewWithBackgroundCompute && PreviewWithBackgroundCompute->HaveValidResult());
+	return (PreviewWithBackgroundCompute && PreviewWithBackgroundCompute->HaveValidNonEmptyResult());
 }
 
 
