AutoLOD: don't clear the default value of thicken weight map name

#rnx

diff --git a/Engine/Plugins/Experimental/MeshLODToolset/Source/MeshLODToolset/Private/Graphs/GenerateStaticMeshLODProcess.cpp b/Engine/Plugins/Experimental/MeshLODToolset/Source/MeshLODToolset/Private/Graphs/GenerateStaticMeshLODProcess.cpp
--- a/Engine/Plugins/Experimental/MeshLODToolset/Source/MeshLODToolset/Private/Graphs/GenerateStaticMeshLODProcess.cpp
+++ b/Engine/Plugins/Experimental/MeshLODToolset/Source/MeshLODToolset/Private/Graphs/GenerateStaticMeshLODProcess.cpp
@@ -601,7 +601,6 @@ bool UGenerateStaticMeshLODProcess::InitializeGenerator()
 
 	CurrentSettings_Preprocess.FilterGroupLayer = Generator->GetCurrentPreFilterSettings().FilterGroupLayerName;
 	CurrentSettings_Preprocess.ThickenAmount = Generator->GetCurrentThickenSettings().ThickenAmount;
-	CurrentSettings_Preprocess.ThickenWeightMapName = "";
 	CurrentSettings.MeshGenerator = static_cast<EGenerateStaticMeshLODProcess_MeshGeneratorModes>(Generator->GetCurrentCoreMeshGeneratorMode());
 
 	CurrentSettings.SolidifyVoxelResolution = Generator->GetCurrentSolidifySettings().VoxelResolution;
