AutoLOD: release the Process object when the Tool shuts down so that it can be garbage collected without waiting for the Tool to be deleted

#rnx

diff --git a/Engine/Plugins/Experimental/MeshLODToolset/Source/MeshLODToolset/Private/Tools/GenerateStaticMeshLODAssetTool.cpp b/Engine/Plugins/Experimental/MeshLODToolset/Source/MeshLODToolset/Private/Tools/GenerateStaticMeshLODAssetTool.cpp
--- a/Engine/Plugins/Experimental/MeshLODToolset/Source/MeshLODToolset/Private/Tools/GenerateStaticMeshLODAssetTool.cpp
+++ b/Engine/Plugins/Experimental/MeshLODToolset/Source/MeshLODToolset/Private/Tools/GenerateStaticMeshLODAssetTool.cpp
@@ -496,6 +496,11 @@ void UGenerateStaticMeshLODAssetTool::Shutdown(EToolShutdownType ShutdownType)
 	{
 		PreviewWithBackgroundCompute->Shutdown();
 	}
+
+	if (GenerateProcess)
+	{
+		GenerateProcess = nullptr;
+	}
 }
 
 
