ModelingTools: disable ResetToDefault on various properties of LODManagerTool

diff --git a/Engine/Plugins/Experimental/MeshLODToolset/Source/MeshLODToolset/Public/Tools/LODManagerTool.h b/Engine/Plugins/Experimental/MeshLODToolset/Source/MeshLODToolset/Public/Tools/LODManagerTool.h
--- a/Engine/Plugins/Experimental/MeshLODToolset/Source/MeshLODToolset/Public/Tools/LODManagerTool.h
+++ b/Engine/Plugins/Experimental/MeshLODToolset/Source/MeshLODToolset/Public/Tools/LODManagerTool.h
@@ -37,10 +37,10 @@ struct MESHLODTOOLSET_API FLODManagerLODInfo
 {
 	GENERATED_BODY()
 
-	UPROPERTY(VisibleAnywhere, Category = LODInformation)
+	UPROPERTY(VisibleAnywhere, Category = LODInformation, meta=(NoResetToDefault))
 	int32 VertexCount = 0;
 
-	UPROPERTY(VisibleAnywhere, Category = LODInformation)
+	UPROPERTY(VisibleAnywhere, Category = LODInformation, meta=(NoResetToDefault))
 	int32 TriangleCount = 0;
 };
 
@@ -51,23 +51,23 @@ class MESHLODTOOLSET_API ULODManagerLODProperties : public UInteractiveToolPrope
 {
 	GENERATED_BODY()
 public:
-	UPROPERTY(VisibleAnywhere, Category = LODInformation)
+	UPROPERTY(VisibleAnywhere, Category = LODInformation, meta = (NoResetToDefault))
 	TArray<FLODManagerLODInfo> SourceLODs;
 
-	UPROPERTY(VisibleAnywhere, Category = LODInformation)
+	UPROPERTY(VisibleAnywhere, Category = LODInformation, meta = (NoResetToDefault))
 	TArray<FLODManagerLODInfo> HiResSource;
 
 
-	UPROPERTY(VisibleAnywhere, Category = LODInformation)
+	UPROPERTY(VisibleAnywhere, Category = LODInformation, meta = (NoResetToDefault))
 	TArray<FLODManagerLODInfo> RenderLODs;
 
-	UPROPERTY(VisibleAnywhere, Category = Nanite, meta = (DisplayName="Enabled"))
+	UPROPERTY(VisibleAnywhere, Category = Nanite, meta = (NoResetToDefault, DisplayName="Enabled"))
 	bool bNaniteEnabled = false;
 
-	UPROPERTY(VisibleAnywhere, Category = Nanite)
+	UPROPERTY(VisibleAnywhere, Category = Nanite, meta = (NoResetToDefault))
 	float PercentTriangles = 0.0f;
 
-	UPROPERTY(VisibleAnywhere, Category = Materials)
+	UPROPERTY(VisibleAnywhere, Category = Materials, meta = (NoResetToDefault))
 	TArray<FStaticMaterial> Materials;
 
 };
@@ -80,7 +80,7 @@ class MESHLODTOOLSET_API ULODManagerPreviewLODProperties : public UInteractiveTo
 	GENERATED_BODY()
 public:
 	/** LOD to visualise. Default option is equivalent to disabling the Tool, RenderData is the mesh used for rendering derived from the SourceModel (possibly simplified) */
-	UPROPERTY(EditAnywhere, Category = LODPreview, meta = (DisplayName = "Show LOD", GetOptions = GetLODNamesFunc))
+	UPROPERTY(EditAnywhere, Category = LODPreview, meta = (DisplayName = "Show LOD", NoResetToDefault, GetOptions = GetLODNamesFunc))
 	FString VisibleLOD;
 
 	UFUNCTION()
