AutoLOD tool UI fixes:
- add a Detail Customization class for AutoLOD properties to enable HintText to be used in some text entry fields when the field is empty
- validate group layer name and weight map layer name when specified by the user. Issue a warning if the specified layer doesn't exist
- don't clear the warning area if we have a non-empty result

diff --git a/Engine/Plugins/Experimental/GeometryFlow/Source/GeometryFlowMeshProcessing/Public/MeshDecompositionNodes/MakeTriangleSetsNode.h b/Engine/Plugins/Experimental/GeometryFlow/Source/GeometryFlowMeshProcessing/Public/MeshDecompositionNodes/MakeTriangleSetsNode.h
--- a/Engine/Plugins/Experimental/GeometryFlow/Source/GeometryFlowMeshProcessing/Public/MeshDecompositionNodes/MakeTriangleSetsNode.h
+++ b/Engine/Plugins/Experimental/GeometryFlow/Source/GeometryFlowMeshProcessing/Public/MeshDecompositionNodes/MakeTriangleSetsNode.h
@@ -93,7 +93,7 @@ public:
 	{
 		AddInput(InParamGroupLayer(), 
 				 MakeUnique<TBasicNodeInput<FName, (int)EDataTypes::Name>>(), 
-				 MakeSafeShared<TMovableData<FName, (int)EDataTypes::Name>>("Default"));
+				 MakeSafeShared<TMovableData<FName, (int)EDataTypes::Name>>());
 
 		AddInput(InParamIgnoreGroups(), 
 				 MakeBasicInput<FIndexSets>());
@@ -167,6 +167,10 @@ public:
 		TSafeSharedPtr<IData> GroupLayerArg = DatasIn.FindData(InParamGroupLayer());
 		FName GroupName = GroupLayerArg->GetDataConstRef<FName>((int)EDataTypes::Name);
 
+		if (GroupName.IsNone())
+		{
+			return;
+		}
 		if (Mesh.HasTriangleGroups() && (GroupName == "Default"))
 		{
 			ComputeIndexSetsForGroups(Mesh, IgnoreGroups, [&Mesh](int TriangleID) { return Mesh.GetTriangleGroup(TriangleID); }, SetsOut);
diff --git a/Engine/Plugins/Experimental/MeshLODToolset/Source/MeshLODToolset/MeshLODToolset.Build.cs b/Engine/Plugins/Experimental/MeshLODToolset/Source/MeshLODToolset/MeshLODToolset.Build.cs
--- a/Engine/Plugins/Experimental/MeshLODToolset/Source/MeshLODToolset/MeshLODToolset.Build.cs
+++ b/Engine/Plugins/Experimental/MeshLODToolset/Source/MeshLODToolset/MeshLODToolset.Build.cs
@@ -51,6 +51,7 @@ public class MeshLODToolset : ModuleRules
 				"UnrealEd",
 				"EditorScriptingUtilities",
 				"Slate",
+				"SlateCore",
 				// ... add private dependencies that you statically link with here ...
 			}
 		);
diff --git a/Engine/Plugins/Experimental/MeshLODToolset/Source/MeshLODToolset/Private/Graphs/GenerateMeshLODGraph.cpp b/Engine/Plugins/Experimental/MeshLODToolset/Source/MeshLODToolset/Private/Graphs/GenerateMeshLODGraph.cpp
--- a/Engine/Plugins/Experimental/MeshLODToolset/Source/MeshLODToolset/Private/Graphs/GenerateMeshLODGraph.cpp
+++ b/Engine/Plugins/Experimental/MeshLODToolset/Source/MeshLODToolset/Private/Graphs/GenerateMeshLODGraph.cpp
@@ -490,7 +490,7 @@ void FGenerateMeshLODGraph::BuildGraph(const FDynamicMesh3* SourceMeshHint)
 	UpdateSettingsSourceNodeValue(*Graph, FilterGroupsNode, IgnoreGroupsForDelete);
 
 	FMeshLODGraphPreFilterSettings PreFilterSettings;
-	PreFilterSettings.FilterGroupLayerName = FName(TEXT("PreFilterGroups"));
+	PreFilterSettings.FilterGroupLayerName = FName();
 	UpdatePreFilterSettings(PreFilterSettings);
 
 	UE::GeometryFlow::FMeshThickenSettings MeshThickenSettings;
diff --git a/Engine/Plugins/Experimental/MeshLODToolset/Source/MeshLODToolset/Private/MeshLODToolsetModule.cpp b/Engine/Plugins/Experimental/MeshLODToolset/Source/MeshLODToolset/Private/MeshLODToolsetModule.cpp
--- a/Engine/Plugins/Experimental/MeshLODToolset/Source/MeshLODToolset/Private/MeshLODToolsetModule.cpp
+++ b/Engine/Plugins/Experimental/MeshLODToolset/Source/MeshLODToolset/Private/MeshLODToolsetModule.cpp
@@ -3,6 +3,10 @@
 #include "MeshLODToolsetModule.h"
 #include "Tools/LODGenerationSettingsAsset.h"
 #include "AssetTypeActions_Base.h"
+#include "Modules/ModuleManager.h"
+#include "PropertyEditorModule.h"
+#include "Tools/DetailsCustomizations/AutoLODToolCustomizations.h"
+#include "Tools/GenerateStaticMeshLODAssetTool.h"
 
 #define LOCTEXT_NAMESPACE "FMeshLODToolsetModule"
 
@@ -24,10 +28,10 @@ public:
 
 
 
-
+// This function will execute after your module is loaded into memory; the exact timing is specified in the .uplugin file per-module
 void FMeshLODToolsetModule::StartupModule()
 {
-	// This code will execute after your module is loaded into memory; the exact timing is specified in the .uplugin file per-module
+	FCoreDelegates::OnPostEngineInit.AddRaw(this, &FMeshLODToolsetModule::OnPostEngineInit);
 
 	// Register asset actions
 	IAssetTools& AssetTools = FModuleManager::LoadModuleChecked<FAssetToolsModule>("AssetTools").Get();
@@ -38,11 +42,10 @@ void FMeshLODToolsetModule::StartupModule()
 	}
 }
 
+// This function may be called during shutdown to clean up your module.  For modules that support dynamic reloading,
+// we call this function before unloading the module.
 void FMeshLODToolsetModule::ShutdownModule()
 {
-	// This function may be called during shutdown to clean up your module.  For modules that support dynamic reloading,
-	// we call this function before unloading the module.
-
 	// Unregister asset actions
 	FAssetToolsModule* AssetToolsModule = FModuleManager::GetModulePtr<FAssetToolsModule>("AssetTools");
 	if (AssetToolsModule)
@@ -52,6 +55,26 @@ void FMeshLODToolsetModule::ShutdownModule()
 			AssetToolsModule->Get().UnregisterAssetTypeActions(RegisteredAssetTypeAction);
 		}
 	}
+
+	// Unregister customizations
+	FPropertyEditorModule* PropertyEditorModule = FModuleManager::GetModulePtr<FPropertyEditorModule>("PropertyEditor");
+	if (PropertyEditorModule)
+	{
+		for (FName ClassName : ClassesToUnregisterOnShutdown)
+		{
+			PropertyEditorModule->UnregisterCustomClassLayout(ClassName);
+		}
+	}
+}
+
+void FMeshLODToolsetModule::OnPostEngineInit()
+{
+	ClassesToUnregisterOnShutdown.Reset();
+
+	FPropertyEditorModule& PropertyModule = FModuleManager::LoadModuleChecked<FPropertyEditorModule>("PropertyEditor");
+
+	PropertyModule.RegisterCustomClassLayout("GenerateStaticMeshLODAssetToolProperties", FOnGetDetailCustomizationInstance::CreateStatic(&FAutoLODToolDetails::MakeInstance));
+	ClassesToUnregisterOnShutdown.Add(UGenerateStaticMeshLODAssetToolProperties::StaticClass()->GetFName());
 }
 
 #undef LOCTEXT_NAMESPACE
diff --git a/Engine/Plugins/Experimental/MeshLODToolset/Source/MeshLODToolset/Private/Tools/DetailsCustomizations/AutoLODToolCustomizations.cpp b/Engine/Plugins/Experimental/MeshLODToolset/Source/MeshLODToolset/Private/Tools/DetailsCustomizations/AutoLODToolCustomizations.cpp
new file mode 100644
--- /dev/null
+++ b/Engine/Plugins/Experimental/MeshLODToolset/Source/MeshLODToolset/Private/Tools/DetailsCustomizations/AutoLODToolCustomizations.cpp
@@ -0,0 +1,84 @@
+// Copyright Epic Games, Inc. All Rights Reserved.
+
+#include "Tools/DetailsCustomizations/AutoLODToolCustomizations.h"
+#include "UObject/Class.h"
+#include "DetailLayoutBuilder.h"
+#include "DetailCategoryBuilder.h"
+#include "DetailWidgetRow.h"
+#include "IDetailGroup.h"
+#include "PropertyHandle.h"
+#include "PropertyRestriction.h"
+#include "Widgets/Input/SEditableTextBox.h"
+#include "Tools/GenerateStaticMeshLODAssetTool.h"
+
+#define LOCTEXT_NAMESPACE "AutoLODToolCustomizations"
+
+TSharedRef<IDetailCustomization> FAutoLODToolDetails::MakeInstance()
+{
+	return MakeShareable(new FAutoLODToolDetails);
+}
+
+void FAutoLODToolDetails::CustomizeDetails(IDetailLayoutBuilder& DetailBuilder)
+{
+	TSharedRef<IPropertyHandle> PreprocessingGroupSettings = DetailBuilder.GetProperty(GET_MEMBER_NAME_CHECKED(UGenerateStaticMeshLODAssetToolProperties, Preprocessing), 
+		UGenerateStaticMeshLODAssetToolProperties::StaticClass());
+	
+	// NOTE: We are completely rebuilding the widgets for the FGenerateStaticMeshLODProcess_PreprocessSettings struct. If the struct changes, this code will
+	// need to be updated.
+	 	
+	// Hide existing group
+	DetailBuilder.HideProperty(PreprocessingGroupSettings);
+
+	// Create new group
+	IDetailCategoryBuilder& ConfigCategory = DetailBuilder.EditCategory("Generator Configuration");
+
+	IDetailGroup& DetailGroup = ConfigCategory.AddGroup("Preprocessing", LOCTEXT("Preprocessing", "Preprocessing"));
+
+	TSharedPtr<IPropertyHandle> FilterGroupLayerHandle = PreprocessingGroupSettings->GetChildHandle(FName("FilterGroupLayer"));
+	check(FilterGroupLayerHandle);
+
+	DetailGroup.AddWidgetRow()
+	.NameContent()
+	[
+		FilterGroupLayerHandle->CreatePropertyNameWidget(LOCTEXT("DetailFilterGroupLayerName", "Detail Filter Group Layer"))
+	]
+	.ValueContent()
+	.VAlign(VAlign_Center)
+	[
+		SNew(SEditableTextBox)
+		.Text(FText())
+		.HintText(LOCTEXT("DetailFilterGroupLayerValueHintText", "(e.g. PreFilterGroups)"))
+		.OnTextCommitted_Lambda([FilterGroupLayerHandle](const FText& Val, ETextCommit::Type TextCommitType)
+		{
+			FilterGroupLayerHandle->SetValue(Val.ToString());
+		})
+	];
+
+	TSharedPtr<IPropertyHandle> ThickenWeightMapNameHandle = PreprocessingGroupSettings->GetChildHandle(FName("ThickenWeightMapName"));
+	check(ThickenWeightMapNameHandle);
+
+	DetailGroup.AddWidgetRow()
+	.NameContent()
+	[
+		ThickenWeightMapNameHandle->CreatePropertyNameWidget(LOCTEXT("ThickenWeightMapName", "Thicken Weight Map"))
+	]
+	.ValueContent()
+	.VAlign(VAlign_Center)
+	[
+		SNew(SEditableTextBox)
+		.Text(FText())
+		.HintText(LOCTEXT("ThickenWeightMapNameValueHintText", "(e.g. ThickenMap)"))
+		.OnTextCommitted_Lambda([ThickenWeightMapNameHandle](const FText& Val, ETextCommit::Type TextCommitType)
+		{
+			ThickenWeightMapNameHandle->SetValue(Val.ToString());
+		})
+	];
+
+	TSharedPtr<IPropertyHandle> ThickenAmountHandle = PreprocessingGroupSettings->GetChildHandle(FName("ThickenAmount"));
+	check(ThickenAmountHandle);
+
+	DetailGroup.AddPropertyRow(ThickenAmountHandle.ToSharedRef());
+}
+
+#undef LOCTEXT_NAMESPACE
+
diff --git a/Engine/Plugins/Experimental/MeshLODToolset/Source/MeshLODToolset/Private/Tools/GenerateStaticMeshLODAssetTool.cpp b/Engine/Plugins/Experimental/MeshLODToolset/Source/MeshLODToolset/Private/Tools/GenerateStaticMeshLODAssetTool.cpp
--- a/Engine/Plugins/Experimental/MeshLODToolset/Source/MeshLODToolset/Private/Tools/GenerateStaticMeshLODAssetTool.cpp
+++ b/Engine/Plugins/Experimental/MeshLODToolset/Source/MeshLODToolset/Private/Tools/GenerateStaticMeshLODAssetTool.cpp
@@ -378,8 +378,14 @@ void UGenerateStaticMeshLODAssetTool::Setup()
 	ToolSetupUtil::ApplyRenderingConfigurationToPreview(PreviewWithBackgroundCompute->PreviewMesh, nullptr);
 	PreviewWithBackgroundCompute->PreviewMesh->SetTangentsMode(EDynamicMeshComponentTangentsMode::ExternallyProvided);
 
-	PreviewWithBackgroundCompute->OnMeshUpdated.AddLambda([this](UMeshOpPreviewWithBackgroundCompute* Compute) {
-		UpdateAcceptWarnings(Compute->HaveEmptyResult() ? EAcceptWarning::EmptyForbidden : EAcceptWarning::NoWarning);
+	PreviewWithBackgroundCompute->OnMeshUpdated.AddLambda([this](UMeshOpPreviewWithBackgroundCompute* Compute) 
+	{
+		if (Compute->HaveEmptyResult())
+		{
+			GetToolManager()->DisplayMessage(LOCTEXT("CannotCreateEmptyMesh", "WARNING: Tool doesn't allow creation of an empty mesh."),
+				EToolMessageLevel::UserWarning);
+		}
+		// Don't clear the message area if we have a non-empty result
 	});
 
 	// For the first computation, display a bounding box with the working material. Otherwise it looks like nothing
@@ -451,8 +457,81 @@ void UGenerateStaticMeshLODAssetTool::Setup()
 	}
 }
 
+
+
+bool UGenerateStaticMeshLODAssetTool::ValidateSettings() const
+{
+	if (!BasicProperties || !GenerateProcess)
+	{
+		return true;
+	}
+
+	const FDynamicMesh3& Mesh = GenerateProcess->GetSourceMesh();
+
+	const FName& GroupName = BasicProperties->Preprocessing.FilterGroupLayer;
+	if (!GroupName.IsNone())
+	{
+		bool bFound = false;
+
+		if (Mesh.HasAttributes())
+		{
+			for (int GroupLayerIndex = 0; GroupLayerIndex < Mesh.Attributes()->NumPolygroupLayers(); ++GroupLayerIndex)
+			{
+				const FDynamicMeshPolygroupAttribute* GroupAttribute = Mesh.Attributes()->GetPolygroupLayer(GroupLayerIndex);
+				if (GroupAttribute->GetName() == GroupName)
+				{
+					bFound = true;
+					break;
+				}
+			}
+		}
+
+		if (!bFound)
+		{
+			const FText Message = FText::Format(LOCTEXT("GroupNotFoundWarning", "Group {0} not found on input mesh"), FText::FromName(GroupName));
+			GetToolManager()->DisplayMessage(Message, EToolMessageLevel::UserWarning);
+			return false;
+		}
+	}
+
+
+	const FName& WeightMapName = BasicProperties->Preprocessing.ThickenWeightMapName;
+	if (!WeightMapName.IsNone())
+	{
+		bool bFound = false;
+
+		if (Mesh.HasAttributes())
+		{
+			for (int WeightMapIndex = 0; WeightMapIndex < Mesh.Attributes()->NumWeightLayers(); ++WeightMapIndex)
+			{
+				const FDynamicMeshWeightAttribute* GroupAttribute = Mesh.Attributes()->GetWeightLayer(WeightMapIndex);
+				if (GroupAttribute->GetName() == WeightMapName)
+				{
+					bFound = true;
+					break;
+				}
+			}
+		}
+
+		if (!bFound)
+		{
+			const FText Message = FText::Format(LOCTEXT("WeightMapNotFoundWarning", "Weight Map {0} not found on input mesh"), FText::FromName(WeightMapName));
+			GetToolManager()->DisplayMessage(Message, EToolMessageLevel::UserWarning);
+			return false;
+		}
+	}
+
+	return true;
+}
+
 void UGenerateStaticMeshLODAssetTool::OnSettingsModified()
 {
+	bool bOK = ValidateSettings();
+	if (bOK)
+	{
+		GetToolManager()->DisplayMessage({}, EToolMessageLevel::UserWarning);
+	}
+
 	PreviewWithBackgroundCompute->InvalidateResult();
 }
 
diff --git a/Engine/Plugins/Experimental/MeshLODToolset/Source/MeshLODToolset/Public/Graphs/GenerateMeshLODGraph.h b/Engine/Plugins/Experimental/MeshLODToolset/Source/MeshLODToolset/Public/Graphs/GenerateMeshLODGraph.h
--- a/Engine/Plugins/Experimental/MeshLODToolset/Source/MeshLODToolset/Public/Graphs/GenerateMeshLODGraph.h
+++ b/Engine/Plugins/Experimental/MeshLODToolset/Source/MeshLODToolset/Public/Graphs/GenerateMeshLODGraph.h
@@ -26,7 +26,7 @@ using UE::Geometry::FDynamicMesh3;
 
 struct FMeshLODGraphPreFilterSettings
 {
-	FName FilterGroupLayerName = FName( "PreFilterGroups" );
+	FName FilterGroupLayerName = FName();
 };
 
 class FGenerateMeshLODGraph
diff --git a/Engine/Plugins/Experimental/MeshLODToolset/Source/MeshLODToolset/Public/Graphs/GenerateStaticMeshLODProcess.h b/Engine/Plugins/Experimental/MeshLODToolset/Source/MeshLODToolset/Public/Graphs/GenerateStaticMeshLODProcess.h
--- a/Engine/Plugins/Experimental/MeshLODToolset/Source/MeshLODToolset/Public/Graphs/GenerateStaticMeshLODProcess.h
+++ b/Engine/Plugins/Experimental/MeshLODToolset/Source/MeshLODToolset/Public/Graphs/GenerateStaticMeshLODProcess.h
@@ -79,20 +79,25 @@ struct FGenerateStaticMeshLODProcess_PreprocessSettings
 {
 	GENERATED_BODY()
 
+	//
+	// NOTE: Customization of widgets for these properties happens in FAutoLODToolDetails::CustomizeDetails(). If you add, remove, or 
+	// change any of the properties in this struct, you must also update that function!
+	//
+
 	// Filter settings
 
 	/** Group layer to use for filtering out detail before processing */
-	UPROPERTY(EditAnywhere, Category = DetailFilter, meta = (DisplayName = "Detail Filter Group Layer"))
-	FName FilterGroupLayer = FName(TEXT("PreFilterGroups"));
+	UPROPERTY(EditAnywhere, Category = Preprocessing, meta = (DisplayName = "Detail Filter Group Layer"))
+	FName FilterGroupLayer = FName();
 
 	// Thicken settings
 
 	/** Weight map used during mesh thickening */
-	UPROPERTY(EditAnywhere, Category = DetailFilter, meta = (DisplayName = "Thicken Weight Map"))
-	FName ThickenWeightMapName = FName(TEXT("ThickenWeightMap"));
+	UPROPERTY(EditAnywhere, Category = Preprocessing, meta = (DisplayName = "Thicken Weight Map"))
+	FName ThickenWeightMapName = FName();
 
 	/** Amount to thicken the mesh prior to Solidifying. The thicken weight map values are multiplied by this value. */
-	UPROPERTY(EditAnywhere, Category = DetailFilter, meta = (UIMin = "0", UIMax = "100", ClampMin = "0", ClampMax = "1000"))
+	UPROPERTY(EditAnywhere, Category = Preprocessing, meta = (UIMin = "0", UIMax = "100", ClampMin = "0", ClampMax = "1000"))
 	float ThickenAmount = 0.0f;
 };
 
diff --git a/Engine/Plugins/Experimental/MeshLODToolset/Source/MeshLODToolset/Public/MeshLODToolsetModule.h b/Engine/Plugins/Experimental/MeshLODToolset/Source/MeshLODToolset/Public/MeshLODToolsetModule.h
--- a/Engine/Plugins/Experimental/MeshLODToolset/Source/MeshLODToolset/Public/MeshLODToolsetModule.h
+++ b/Engine/Plugins/Experimental/MeshLODToolset/Source/MeshLODToolset/Public/MeshLODToolsetModule.h
@@ -18,7 +18,13 @@ public:
 	virtual void StartupModule() override;
 	virtual void ShutdownModule() override;
 
+	void OnPostEngineInit();
 
 	// registered asset actions
 	TArray<TSharedRef<IAssetTypeActions>> RegisteredAssetTypeActions;
+
+private:
+	/** StaticClass is not safe on shutdown, so we cache the name, and use this to unregister on shut down */
+	TArray<FName> ClassesToUnregisterOnShutdown;
+
 };
diff --git a/Engine/Plugins/Experimental/MeshLODToolset/Source/MeshLODToolset/Public/Tools/DetailsCustomizations/AutoLODToolCustomizations.h b/Engine/Plugins/Experimental/MeshLODToolset/Source/MeshLODToolset/Public/Tools/DetailsCustomizations/AutoLODToolCustomizations.h
new file mode 100644
--- /dev/null
+++ b/Engine/Plugins/Experimental/MeshLODToolset/Source/MeshLODToolset/Public/Tools/DetailsCustomizations/AutoLODToolCustomizations.h
@@ -0,0 +1,18 @@
+// Copyright Epic Games, Inc. All Rights Reserved.
+
+#pragma once
+
+#include "IDetailCustomization.h"
+#include "Templates/SharedPointer.h"
+
+/// 
+/// Detail customization for UGenerateStaticMeshLODAssetTool (AKA AutoLOD) properties
+///
+ 
+class FAutoLODToolDetails : public IDetailCustomization
+{
+public:
+	static TSharedRef<IDetailCustomization> MakeInstance();
+	virtual void CustomizeDetails(IDetailLayoutBuilder& DetailBuilder) override;
+};
+
diff --git a/Engine/Plugins/Experimental/MeshLODToolset/Source/MeshLODToolset/Public/Tools/GenerateStaticMeshLODAssetTool.h b/Engine/Plugins/Experimental/MeshLODToolset/Source/MeshLODToolset/Public/Tools/GenerateStaticMeshLODAssetTool.h
--- a/Engine/Plugins/Experimental/MeshLODToolset/Source/MeshLODToolset/Public/Tools/GenerateStaticMeshLODAssetTool.h
+++ b/Engine/Plugins/Experimental/MeshLODToolset/Source/MeshLODToolset/Public/Tools/GenerateStaticMeshLODAssetTool.h
@@ -343,4 +343,5 @@ protected:
 
 	void OnPresetSelectionChanged();
 
+	bool ValidateSettings() const;
 };
