Do not auto-load settings on switching the target preset in AutoLOD, because it's impossible to tell whether the user is picking the target to write their current settings to, or to overwrite their current settings from.

diff --git a/Engine/Plugins/Experimental/MeshLODToolset/Source/MeshLODToolset/Private/Tools/GenerateStaticMeshLODAssetTool.cpp b/Engine/Plugins/Experimental/MeshLODToolset/Source/MeshLODToolset/Private/Tools/GenerateStaticMeshLODAssetTool.cpp
--- a/Engine/Plugins/Experimental/MeshLODToolset/Source/MeshLODToolset/Private/Tools/GenerateStaticMeshLODAssetTool.cpp
+++ b/Engine/Plugins/Experimental/MeshLODToolset/Source/MeshLODToolset/Private/Tools/GenerateStaticMeshLODAssetTool.cpp
@@ -443,8 +443,9 @@ void UGenerateStaticMeshLODAssetTool::Setup()
 	CollisionPreview = NewObject<UPreviewGeometry>(this);
 	CollisionPreview->CreateInWorld(GetTargetWorld(), (FTransform)PreviewTransform);
 
-	// read Preset if we started Tool with one already selected
-	OnPresetSelectionChanged();		
+	// Trigger any automatic Preset-changed behavior if we started Tool with one already selected
+	// Note: Currently this does nothing, as we rely on the user to manually read/write the preset
+	OnPresetSelectionChanged();
 
 	// Pop up notifications for any warnings
 	for ( const FProgressCancel::FMessageInfo& Warning : Progress.Warnings )
@@ -676,7 +677,18 @@ void UGenerateStaticMeshLODAssetTool::RequestPresetAction(EGenerateLODAssetToolP
 
 	if (ActionType == EGenerateLODAssetToolPresetAction::ReadFromPreset)
 	{
-		OnPresetSelectionChanged();
+		UStaticMeshLODGenerationSettings* ApplyPreset = PresetProperties->Preset.Get();
+		if (ApplyPreset)
+		{
+			BasicProperties->Preprocessing = ApplyPreset->Preprocessing;
+			BasicProperties->MeshGeneration = ApplyPreset->MeshGeneration;
+			BasicProperties->Simplification = ApplyPreset->Simplification;
+			BasicProperties->Normals = ApplyPreset->Normals;
+			BasicProperties->TextureBaking = ApplyPreset->TextureBaking;
+			BasicProperties->UVGeneration = ApplyPreset->UVGeneration;
+			BasicProperties->SimpleCollision = ApplyPreset->SimpleCollision;
+			OnSettingsModified();
+		}
 	}
 	else if (ActionType == EGenerateLODAssetToolPresetAction::WriteToPreset)
 	{
@@ -702,18 +714,7 @@ void UGenerateStaticMeshLODAssetTool::RequestPresetAction(EGenerateLODAssetToolP
 
 void UGenerateStaticMeshLODAssetTool::OnPresetSelectionChanged()
 {
-	UStaticMeshLODGenerationSettings* ApplyPreset = PresetProperties->Preset.Get();
-	if (ApplyPreset)
-	{
-		BasicProperties->Preprocessing = ApplyPreset->Preprocessing;
-		BasicProperties->MeshGeneration = ApplyPreset->MeshGeneration;
-		BasicProperties->Simplification = ApplyPreset->Simplification;
-		BasicProperties->Normals = ApplyPreset->Normals;
-		BasicProperties->TextureBaking = ApplyPreset->TextureBaking;
-		BasicProperties->UVGeneration = ApplyPreset->UVGeneration;
-		BasicProperties->SimpleCollision = ApplyPreset->SimpleCollision;
-		OnSettingsModified();
-	}
+	// Rely on user to decide when to write/read settings to/from the selected preset
 }
 
 
