StaticMeshLODGenerationSettings: convert AssetTypeActions to AssetDefinition

diff --git a/Engine/Plugins/Experimental/MeshLODToolset/Source/MeshLODToolset/MeshLODToolset.Build.cs b/Engine/Plugins/Experimental/MeshLODToolset/Source/MeshLODToolset/MeshLODToolset.Build.cs
--- a/Engine/Plugins/Experimental/MeshLODToolset/Source/MeshLODToolset/MeshLODToolset.Build.cs
+++ b/Engine/Plugins/Experimental/MeshLODToolset/Source/MeshLODToolset/MeshLODToolset.Build.cs
@@ -52,6 +52,7 @@ public class MeshLODToolset : ModuleRules
 				"EditorScriptingUtilities",
 				"Slate",
 				"SlateCore",
+				"AssetDefinition"
 				// ... add private dependencies that you statically link with here ...
 			}
 		);
diff --git a/Engine/Plugins/Experimental/MeshLODToolset/Source/MeshLODToolset/Private/MeshLODToolsetModule.cpp b/Engine/Plugins/Experimental/MeshLODToolset/Source/MeshLODToolset/Private/MeshLODToolsetModule.cpp
--- a/Engine/Plugins/Experimental/MeshLODToolset/Source/MeshLODToolset/Private/MeshLODToolsetModule.cpp
+++ b/Engine/Plugins/Experimental/MeshLODToolset/Source/MeshLODToolset/Private/MeshLODToolsetModule.cpp
@@ -3,59 +3,25 @@
 #include "MeshLODToolsetModule.h"
 #include "Misc/CoreDelegates.h"
 #include "Tools/LODGenerationSettingsAsset.h"
-#include "AssetTypeActions_Base.h"
 #include "PropertyEditorModule.h"
 #include "Tools/DetailsCustomizations/AutoLODToolCustomizations.h"
 #include "Tools/GenerateStaticMeshLODAssetTool.h"
+#include "Modules/ModuleManager.h"
 
 #define LOCTEXT_NAMESPACE "FMeshLODToolsetModule"
 
 DEFINE_LOG_CATEGORY(LogMeshLODToolset);
 
-
-/**
- * Asset Type Actions for UStaticMeshLODGenerationSettings Assets
- */
-class FAssetTypeActions_StaticMeshLODGenerationSettings : public FAssetTypeActions_Base
-{
-public:
-	virtual FText GetName() const override { return NSLOCTEXT("AssetTypeActions", "FAssetTypeActions_StaticMeshLODGenerationSettings", "AutoLOD Settings"); }
-	virtual FColor GetTypeColor() const override { return FColor( 175, 0, 128 ); }
-	virtual UClass* GetSupportedClass() const override { return UStaticMeshLODGenerationSettings::StaticClass(); }
-	virtual uint32 GetCategories() override { return EAssetTypeCategories::Misc; }
-	virtual bool CanLocalize() const override { return false; }
-};
-
-
-
 // This function will execute after your module is loaded into memory; the exact timing is specified in the .uplugin file per-module
 void FMeshLODToolsetModule::StartupModule()
 {
 	FCoreDelegates::OnPostEngineInit.AddRaw(this, &FMeshLODToolsetModule::OnPostEngineInit);
-
-	// Register asset actions
-	IAssetTools& AssetTools = FModuleManager::LoadModuleChecked<FAssetToolsModule>("AssetTools").Get();
-	{
-		TSharedRef<IAssetTypeActions> StaticMeshLODGenerationSettings_AssetActions = MakeShareable( new FAssetTypeActions_StaticMeshLODGenerationSettings );
-		AssetTools.RegisterAssetTypeActions( StaticMeshLODGenerationSettings_AssetActions );
-		RegisteredAssetTypeActions.Add(StaticMeshLODGenerationSettings_AssetActions);
-	}
 }
 
 // This function may be called during shutdown to clean up your module.  For modules that support dynamic reloading,
 // we call this function before unloading the module.
 void FMeshLODToolsetModule::ShutdownModule()
 {
-	// Unregister asset actions
-	FAssetToolsModule* AssetToolsModule = FModuleManager::GetModulePtr<FAssetToolsModule>("AssetTools");
-	if (AssetToolsModule)
-	{
-		for (TSharedRef<IAssetTypeActions> RegisteredAssetTypeAction : RegisteredAssetTypeActions)
-		{
-			AssetToolsModule->Get().UnregisterAssetTypeActions(RegisteredAssetTypeAction);
-		}
-	}
-
 	// Unregister customizations
 	FPropertyEditorModule* PropertyEditorModule = FModuleManager::GetModulePtr<FPropertyEditorModule>("PropertyEditor");
 	if (PropertyEditorModule)
diff --git a/Engine/Plugins/Experimental/MeshLODToolset/Source/MeshLODToolset/Public/AssetDefinition_StaticMeshLODGenerationSettings.h b/Engine/Plugins/Experimental/MeshLODToolset/Source/MeshLODToolset/Public/AssetDefinition_StaticMeshLODGenerationSettings.h
new file mode 100644
--- /dev/null
+++ b/Engine/Plugins/Experimental/MeshLODToolset/Source/MeshLODToolset/Public/AssetDefinition_StaticMeshLODGenerationSettings.h
@@ -0,0 +1,28 @@
+// Copyright Epic Games, Inc. All Rights Reserved.
+
+#pragma once
+
+#include "AssetDefinitionDefault.h"
+#include "Tools/LODGenerationSettingsAsset.h"
+#include "AssetDefinition_StaticMeshLODGenerationSettings.generated.h"
+
+/**
+ * Asset Definition for UStaticMeshLODGenerationSettings Assets
+ */
+UCLASS()
+class UAssetDefinition_StaticMeshLODGenerationSettings : public UAssetDefinitionDefault
+{
+	GENERATED_BODY()
+
+private:
+	virtual FText GetAssetDisplayName() const override { return NSLOCTEXT("AssetTypeActions", "FAssetTypeActions_StaticMeshLODGenerationSettings", "AutoLOD Settings"); }
+	virtual FLinearColor GetAssetColor() const override { return FLinearColor(FColor(175, 0, 128)); }
+	virtual TSoftClassPtr<UObject> GetAssetClass() const override { return UStaticMeshLODGenerationSettings::StaticClass(); }
+	virtual TConstArrayView<FAssetCategoryPath> GetAssetCategories() const override
+	{
+		static const auto Categories = { EAssetCategoryPaths::Misc };
+		return Categories;
+	}
+	virtual FAssetSupportResponse CanLocalize(const FAssetData& InAsset) const { return FAssetSupportResponse::NotSupported(); }
+};
+
