Make sure we clear the NavCollision before updating it within StaticMesh tools, to make sure the data is computed properly and the DDC is updated

diff --git a/Engine/Plugins/Experimental/GeometryScripting/Source/GeometryScriptingCore/Private/CollisionFunctions.cpp b/Engine/Plugins/Experimental/GeometryScripting/Source/GeometryScriptingCore/Private/CollisionFunctions.cpp
--- a/Engine/Plugins/Experimental/GeometryScripting/Source/GeometryScriptingCore/Private/CollisionFunctions.cpp
+++ b/Engine/Plugins/Experimental/GeometryScripting/Source/GeometryScriptingCore/Private/CollisionFunctions.cpp
@@ -170,8 +170,7 @@ static void SetStaticMeshSimpleCollision(UStaticMesh* StaticMeshAsset, const FKA
 		// rebuild physics meshes
 		BodySetup->CreatePhysicsMeshes();
 
-		// rebuild nav collision (? StaticMeshEditor does this)
-		StaticMeshAsset->CreateNavCollision(/*bIsUpdate=*/true);
+		StaticMeshAsset->RecreateNavCollision();
 
 		// update physics state on all components using this StaticMesh
 		for (FThreadSafeObjectIterator Iter(UStaticMeshComponent::StaticClass()); Iter; ++Iter)
diff --git a/Engine/Plugins/Experimental/MeshLODToolset/Source/MeshLODToolset/Private/Graphs/GenerateStaticMeshLODProcess.cpp b/Engine/Plugins/Experimental/MeshLODToolset/Source/MeshLODToolset/Private/Graphs/GenerateStaticMeshLODProcess.cpp
--- a/Engine/Plugins/Experimental/MeshLODToolset/Source/MeshLODToolset/Private/Graphs/GenerateStaticMeshLODProcess.cpp
+++ b/Engine/Plugins/Experimental/MeshLODToolset/Source/MeshLODToolset/Private/Graphs/GenerateStaticMeshLODProcess.cpp
@@ -1678,8 +1678,7 @@ void UGenerateStaticMeshLODProcess::WriteDerivedStaticMeshAsset()
 
 	// do we need to do a post edit change here??
 
-	// is this necessary? 
-	GeneratedStaticMesh->CreateNavCollision(/*bIsUpdate=*/true);
+	GeneratedStaticMesh->RecreateNavCollision();
 
 	GeneratedStaticMesh->GetOutermost()->GetMetaData()->SetValue(GeneratedStaticMesh, TEXT("StaticMeshLOD.IsGeneratedMesh"), TEXT("true"));
 	GeneratedStaticMesh->GetOutermost()->GetMetaData()->SetValue(GeneratedStaticMesh, TEXT("StaticMeshLOD.SourceAssetPath"), *GetSourceAssetPath());
@@ -1871,8 +1870,7 @@ void UGenerateStaticMeshLODProcess::UpdateSourceStaticMeshAsset(bool bSetNewHDSo
 
 	// do we need to do a post edit change here??
 
-	// is this necessary? 
-	SourceStaticMesh->CreateNavCollision(/*bIsUpdate=*/true);
+	SourceStaticMesh->RecreateNavCollision();
 
 	// save UUID used for generated Assets
 	SourceStaticMesh->GetOutermost()->GetMetaData()->SetValue(SourceStaticMesh, TEXT("StaticMeshLOD.GenerationGUID"), *DerivedAssetGUIDKey);
diff --git a/Engine/Plugins/Experimental/MeshModelingToolsetExp/Source/MeshModelingToolsExp/Private/Physics/SetCollisionGeometryTool.cpp b/Engine/Plugins/Experimental/MeshModelingToolsetExp/Source/MeshModelingToolsExp/Private/Physics/SetCollisionGeometryTool.cpp
--- a/Engine/Plugins/Experimental/MeshModelingToolsetExp/Source/MeshModelingToolsExp/Private/Physics/SetCollisionGeometryTool.cpp
+++ b/Engine/Plugins/Experimental/MeshModelingToolsetExp/Source/MeshModelingToolsExp/Private/Physics/SetCollisionGeometryTool.cpp
@@ -370,8 +370,7 @@ void USetCollisionGeometryTool::OnShutdown(EToolShutdownType ShutdownType)
 			// rebuild physics meshes
 			BodySetup->CreatePhysicsMeshes();
 
-			// rebuild nav collision (? StaticMeshEditor does this)
-			StaticMesh->CreateNavCollision(/*bIsUpdate=*/true);
+			StaticMesh->RecreateNavCollision();
 
 			// update physics state on all components using this StaticMesh
 			for (FThreadSafeObjectIterator Iter(UStaticMeshComponent::StaticClass()); Iter; ++Iter)
diff --git a/Engine/Plugins/Runtime/MeshModelingToolset/Source/ModelingComponents/Private/Physics/ComponentCollisionUtil.cpp b/Engine/Plugins/Runtime/MeshModelingToolset/Source/ModelingComponents/Private/Physics/ComponentCollisionUtil.cpp
--- a/Engine/Plugins/Runtime/MeshModelingToolset/Source/ModelingComponents/Private/Physics/ComponentCollisionUtil.cpp
+++ b/Engine/Plugins/Runtime/MeshModelingToolset/Source/ModelingComponents/Private/Physics/ComponentCollisionUtil.cpp
@@ -73,8 +73,7 @@ void UE::Geometry::UpdateSimpleCollision(
 	// rebuild physics meshes
 	BodySetup->CreatePhysicsMeshes();
 
-	// rebuild nav collision (? StaticMeshEditor does this)
-	StaticMesh->CreateNavCollision(/*bIsUpdate=*/true);
+	StaticMesh->RecreateNavCollision();
 
 	// update physics state on all components using this StaticMesh
 	for (FThreadSafeObjectIterator Iter(UStaticMeshComponent::StaticClass()); Iter; ++Iter)
diff --git a/Engine/Source/Runtime/Engine/Classes/Engine/StaticMesh.h b/Engine/Source/Runtime/Engine/Classes/Engine/StaticMesh.h
--- a/Engine/Source/Runtime/Engine/Classes/Engine/StaticMesh.h
+++ b/Engine/Source/Runtime/Engine/Classes/Engine/StaticMesh.h
@@ -1851,6 +1851,11 @@ public:
 	 */
 	ENGINE_API void CreateNavCollision(const bool bIsUpdate = false);
 
+	/**
+	 * Delete current NavCollision and create a new one if needed
+	 */
+	ENGINE_API void RecreateNavCollision();
+
 	/** Configures this SM as bHasNavigationData = false and clears stored NavCollision */
 	ENGINE_API void MarkAsNotHavingNavigationData();
 
diff --git a/Engine/Source/Runtime/Engine/Private/StaticMesh.cpp b/Engine/Source/Runtime/Engine/Private/StaticMesh.cpp
--- a/Engine/Source/Runtime/Engine/Private/StaticMesh.cpp
+++ b/Engine/Source/Runtime/Engine/Private/StaticMesh.cpp
@@ -7361,6 +7361,12 @@ void UStaticMesh::CreateNavCollision(const bool bIsUpdate)
 	}
 }
 
+void UStaticMesh::RecreateNavCollision()
+{
+	SetNavCollision(nullptr);
+	CreateNavCollision();
+}
+
 void UStaticMesh::SetNavCollision(UNavCollisionBase* InNavCollision)
 {
 	WaitUntilAsyncPropertyReleased(EStaticMeshAsyncProperties::NavCollision);
