MeshModelingToolset: Disable Presets category in AutoLOD tool detail panel

diff --git a/Engine/Plugins/Editor/ModelingToolsEditorMode/Source/ModelingToolsEditorMode/Private/ModelingToolsEditorMode.cpp b/Engine/Plugins/Editor/ModelingToolsEditorMode/Source/ModelingToolsEditorMode/Private/ModelingToolsEditorMode.cpp
--- a/Engine/Plugins/Editor/ModelingToolsEditorMode/Source/ModelingToolsEditorMode/Private/ModelingToolsEditorMode.cpp
+++ b/Engine/Plugins/Editor/ModelingToolsEditorMode/Source/ModelingToolsEditorMode/Private/ModelingToolsEditorMode.cpp
@@ -605,6 +605,7 @@ void UModelingToolsEditorMode::Enter()
 	RegisterTool(ToolManagerCommands.BeginLODManagerTool, TEXT("BeginLODManagerTool"), LODManagerToolBuilder);
 
 	UGenerateStaticMeshLODAssetToolBuilder* GenerateSMLODToolBuilder = NewObject<UGenerateStaticMeshLODAssetToolBuilder>();
+	GenerateSMLODToolBuilder->bInRestrictiveMode = ModelingModeSettings->InRestrictiveMode();
 	RegisterTool(ToolManagerCommands.BeginGenerateStaticMeshLODAssetTool, TEXT("BeginGenerateStaticMeshLODAssetTool"), GenerateSMLODToolBuilder);
 
 	UISMEditorToolBuilder* ISMEditorToolBuilder = NewObject<UISMEditorToolBuilder>();
diff --git a/Engine/Plugins/Experimental/MeshLODToolset/Source/MeshLODToolset/Private/Tools/GenerateStaticMeshLODAssetTool.cpp b/Engine/Plugins/Experimental/MeshLODToolset/Source/MeshLODToolset/Private/Tools/GenerateStaticMeshLODAssetTool.cpp
--- a/Engine/Plugins/Experimental/MeshLODToolset/Source/MeshLODToolset/Private/Tools/GenerateStaticMeshLODAssetTool.cpp
+++ b/Engine/Plugins/Experimental/MeshLODToolset/Source/MeshLODToolset/Private/Tools/GenerateStaticMeshLODAssetTool.cpp
@@ -221,6 +221,7 @@ UMultiSelectionMeshEditingTool* UGenerateStaticMeshLODAssetToolBuilder::CreateNe
 {
 	UGenerateStaticMeshLODAssetTool* NewTool = NewObject<UGenerateStaticMeshLODAssetTool>(SceneState.ToolManager);
 	NewTool->SetUseAssetEditorMode(bUseAssetEditorMode);
+	NewTool->SetRestrictiveMode(bInRestrictiveMode);
 	return NewTool;
 }
 
@@ -247,6 +248,12 @@ void UGenerateStaticMeshLODAssetTool::SetUseAssetEditorMode(bool bEnable)
 }
 
 
+void UGenerateStaticMeshLODAssetTool::SetRestrictiveMode(bool bEnable)
+{
+	bRestrictiveMode = bEnable;
+}
+
+
 void UGenerateStaticMeshLODAssetTool::Setup()
 {
 	using GenerateStaticMeshLODAssetLocals::FGenerateStaticMeshLODAssetOperatorFactory;
@@ -316,11 +323,14 @@ void UGenerateStaticMeshLODAssetTool::Setup()
 	OutputProperties->OutputMode = (bIsInAssetEditorMode) ? EGenerateLODAssetOutputMode::UpdateExistingAsset : EGenerateLODAssetOutputMode::CreateNewAsset;
 	OutputProperties->bShowOutputMode = !bIsInAssetEditorMode;
 
-	PresetProperties = NewObject<UGenerateStaticMeshLODAssetToolPresetProperties>(this);
-	PresetProperties->RestoreProperties(this);
-	PresetProperties->Initialize(this);
-	AddToolPropertySource(PresetProperties);
-	PresetProperties->WatchProperty(PresetProperties->Preset, [this](TWeakObjectPtr<UStaticMeshLODGenerationSettings>) { OnPresetSelectionChanged(); });
+	if (!bRestrictiveMode)
+	{
+		PresetProperties = NewObject<UGenerateStaticMeshLODAssetToolPresetProperties>(this);
+		PresetProperties->RestoreProperties(this);
+		PresetProperties->Initialize(this);
+		AddToolPropertySource(PresetProperties);
+		PresetProperties->WatchProperty(PresetProperties->Preset, [this](TWeakObjectPtr<UStaticMeshLODGenerationSettings>) { OnPresetSelectionChanged(); });
+	}
 
 	SlowTask.EnterProgressFrame(1);
 	BasicProperties = NewObject<UGenerateStaticMeshLODAssetToolProperties>(this);
@@ -655,6 +665,11 @@ void UGenerateStaticMeshLODAssetTool::UpdateExistingAsset()
 
 void UGenerateStaticMeshLODAssetTool::RequestPresetAction(EGenerateLODAssetToolPresetAction ActionType)
 {
+	if (PresetProperties == nullptr)
+	{
+		return;
+	}
+
 	UStaticMeshLODGenerationSettings* CurrentPreset = PresetProperties->Preset.Get();
 	if (CurrentPreset == nullptr)
 	{
diff --git a/Engine/Plugins/Experimental/MeshLODToolset/Source/MeshLODToolset/Public/Tools/GenerateStaticMeshLODAssetTool.h b/Engine/Plugins/Experimental/MeshLODToolset/Source/MeshLODToolset/Public/Tools/GenerateStaticMeshLODAssetTool.h
--- a/Engine/Plugins/Experimental/MeshLODToolset/Source/MeshLODToolset/Public/Tools/GenerateStaticMeshLODAssetTool.h
+++ b/Engine/Plugins/Experimental/MeshLODToolset/Source/MeshLODToolset/Public/Tools/GenerateStaticMeshLODAssetTool.h
@@ -42,6 +42,7 @@ class MESHLODTOOLSET_API UGenerateStaticMeshLODAssetToolBuilder : public UMultiS
 
 public:
 	bool bUseAssetEditorMode = false;
+	bool bInRestrictiveMode = false;
 
 	virtual bool CanBuildTool(const FToolBuilderState& SceneState) const override;
 	virtual UMultiSelectionMeshEditingTool* CreateNewTool(const FToolBuilderState& SceneState) const override;
@@ -264,9 +265,12 @@ class MESHLODTOOLSET_API UGenerateStaticMeshLODAssetTool : public UMultiSelectio
 
 public:
 
-	// Enable UI Customization for running this Tool in the Static Mesh Asset Editor. Must call before Setup.
+	// Enable UI Customization for running this Tool in the Static Mesh Asset Editor. Must be called before Setup.
 	virtual void SetUseAssetEditorMode(bool bEnable);
 
+	// Enable restrictive mode. Must be called before Setup.
+	virtual void SetRestrictiveMode(bool bEnable);
+
 	virtual void Setup() override;
 	virtual void OnShutdown(EToolShutdownType ShutdownType) override;
 
@@ -324,6 +328,7 @@ protected:
 
 protected:
 	bool bIsInAssetEditorMode = false;
+	bool bRestrictiveMode = false;
 
 	UPROPERTY()
 	TObjectPtr<UGenerateStaticMeshLODProcess> GenerateProcess;
diff --git a/Engine/Plugins/Experimental/StaticMeshEditorModeling/Source/StaticMeshEditorModeling/Private/StaticMeshEditorModelingMode.cpp b/Engine/Plugins/Experimental/StaticMeshEditorModeling/Source/StaticMeshEditorModeling/Private/StaticMeshEditorModelingMode.cpp
--- a/Engine/Plugins/Experimental/StaticMeshEditorModeling/Source/StaticMeshEditorModeling/Private/StaticMeshEditorModelingMode.cpp
+++ b/Engine/Plugins/Experimental/StaticMeshEditorModeling/Source/StaticMeshEditorModeling/Private/StaticMeshEditorModelingMode.cpp
@@ -5,6 +5,7 @@
 #include "EdModeInteractiveToolsContext.h"
 #include "ToolTargetManager.h"
 #include "ModelingToolsManagerActions.h"
+#include "ModelingToolsEditorModeSettings.h"
 #include "ToolTargets/StaticMeshComponentToolTarget.h"
 #include "Tools/GenerateStaticMeshLODAssetTool.h"
 #include "Tools/LODManagerTool.h"
@@ -28,8 +29,11 @@ void UStaticMeshEditorModelingMode::Enter()
 
 	const FModelingToolsManagerCommands& ToolManagerCommands = FModelingToolsManagerCommands::Get();
 
+	const UModelingToolsEditorModeSettings* ModelingModeSettings = GetDefault<UModelingToolsEditorModeSettings>();
+
 	UGenerateStaticMeshLODAssetToolBuilder* GenerateStaticMeshLODAssetToolBuilder = NewObject<UGenerateStaticMeshLODAssetToolBuilder>();
 	GenerateStaticMeshLODAssetToolBuilder->bUseAssetEditorMode = true;
+	GenerateStaticMeshLODAssetToolBuilder->bInRestrictiveMode = ModelingModeSettings->InRestrictiveMode();
 	RegisterTool(ToolManagerCommands.BeginGenerateStaticMeshLODAssetTool, TEXT("BeginGenerateStaticMeshLODAssetTool"), GenerateStaticMeshLODAssetToolBuilder);
 	RegisterTool(ToolManagerCommands.BeginLODManagerTool, TEXT("BeginLODManagerTool"), NewObject<ULODManagerToolBuilder>());
 	RegisterTool(ToolManagerCommands.BeginMeshInspectorTool, TEXT("BeginMeshInspectorTool"), NewObject<UMeshInspectorToolBuilder>());
