Fix some bugs in the AutoLOD tool

1. Fix previous collision geometry incorrectly remaining visible when Collision Type was changed in the AutoLOD tool. Now only the collision geometry for the active Collision Type is shown

2. Fix some incorrect tooltips

diff --git a/Engine/Plugins/Experimental/MeshLODToolset/Source/MeshLODToolset/Private/Tools/GenerateStaticMeshLODAssetTool.cpp b/Engine/Plugins/Experimental/MeshLODToolset/Source/MeshLODToolset/Private/Tools/GenerateStaticMeshLODAssetTool.cpp
--- a/Engine/Plugins/Experimental/MeshLODToolset/Source/MeshLODToolset/Private/Tools/GenerateStaticMeshLODAssetTool.cpp
+++ b/Engine/Plugins/Experimental/MeshLODToolset/Source/MeshLODToolset/Private/Tools/GenerateStaticMeshLODAssetTool.cpp
@@ -401,6 +401,9 @@ void UGenerateStaticMeshLODAssetTool::Setup()
 			(const GenerateStaticMeshLODAssetLocals::FGenerateStaticMeshLODAssetOperatorOp*)(Op);
 		check(GenerateLODOp);
 
+		// Clear any previous collision geometry visualization
+		CollisionPreview->RemoveAllLineSets();
+
 		// Must happen on main thread
 		FPhysicsDataCollection PhysicsData;
 		PhysicsData.Geometry = GenerateLODOp->ResultCollision;
diff --git a/Engine/Plugins/Experimental/MeshLODToolset/Source/MeshLODToolset/Public/Graphs/GenerateStaticMeshLODProcess.h b/Engine/Plugins/Experimental/MeshLODToolset/Source/MeshLODToolset/Public/Graphs/GenerateStaticMeshLODProcess.h
--- a/Engine/Plugins/Experimental/MeshLODToolset/Source/MeshLODToolset/Public/Graphs/GenerateStaticMeshLODProcess.h
+++ b/Engine/Plugins/Experimental/MeshLODToolset/Source/MeshLODToolset/Public/Graphs/GenerateStaticMeshLODProcess.h
@@ -26,7 +26,7 @@ using UE::Geometry::FDynamicMesh3;
 UENUM()
 enum class EGenerateStaticMeshLODProcess_MeshGeneratorModes : uint8
 {
-	// note: must keep in sync with FGenerateMeshLODGraph::ECoreMeshGeneratorMode
+	//~NOTE: must keep in sync with FGenerateMeshLODGraph::ECoreMeshGeneratorMode
 	Solidify = 0,
 	SolidifyAndClose = 1,
 	CleanAndSimplify = 2,
@@ -110,7 +110,7 @@ struct FGenerateStaticMeshLODProcess_PreprocessSettings
 UENUM()
 enum class EGenerateStaticMeshLODProcess_SimplifyMethod : uint8
 {
-	// note: must keep in sync with UE::GeometryFlow::EMeshSimplifyTargetType
+	//~NOTE: must keep in sync with UE::GeometryFlow::EMeshSimplifyTargetType
 	TriangleCount = 0,
 	VertexCount = 1,
 	TrianglePercentage = 2,
@@ -268,8 +268,7 @@ struct FGenerateStaticMeshLODProcess_TextureSettings
 UENUM()
 enum class EGenerateStaticMeshLODSimpleCollisionGeometryType : uint8
 {
-	// NOTE: This must be kept in sync with ESimpleCollisionGeometryType in GenerateSimpleCollisionNode.h
-
+	//~NOTE: This must be kept in sync with ESimpleCollisionGeometryType in GenerateSimpleCollisionNode.h
 	AlignedBoxes,
 	OrientedBoxes,
 	MinimalSpheres,
@@ -280,11 +279,11 @@ enum class EGenerateStaticMeshLODSimpleCollisionGeometryType : uint8
 	None
 };
 
-// NOTE: This must be kept in sync with FMeshSimpleShapeApproximation::EProjectedHullAxisMode in MeshSimpleShapeApproximation.h
 
 UENUM()
 enum class EGenerateStaticMeshLODProjectedHullAxisMode : uint8
 {
+	//~NOTE: This must be kept in sync with FMeshSimpleShapeApproximation::EProjectedHullAxisMode in MeshSimpleShapeApproximation.h
 	X = 0,
 	Y = 1,
 	Z = 2,
