AutoLOD: Various fixes to tooltips, edit conditions and reset buttons

diff --git a/Engine/Plugins/Experimental/MeshLODToolset/Source/MeshLODToolset/Private/Tools/GenerateStaticMeshLODAssetTool.cpp b/Engine/Plugins/Experimental/MeshLODToolset/Source/MeshLODToolset/Private/Tools/GenerateStaticMeshLODAssetTool.cpp
--- a/Engine/Plugins/Experimental/MeshLODToolset/Source/MeshLODToolset/Private/Tools/GenerateStaticMeshLODAssetTool.cpp
+++ b/Engine/Plugins/Experimental/MeshLODToolset/Source/MeshLODToolset/Private/Tools/GenerateStaticMeshLODAssetTool.cpp
@@ -528,6 +528,14 @@ bool UGenerateStaticMeshLODAssetTool::ValidateSettings() const
 
 void UGenerateStaticMeshLODAssetTool::OnSettingsModified()
 {
+	if (BasicProperties)
+	{
+		// Set the transient property which controls the edit condition for ThickenAmount
+		// Note: The tool will display an error if ThickenWeightMapName exists but isn't valid
+		BasicProperties->Preprocessing.bIsThickenAmountEnabled = (
+			BasicProperties->Preprocessing.ThickenWeightMapName.IsNone() == false);
+	}
+
 	bool bOK = ValidateSettings();
 	if (bOK)
 	{
diff --git a/Engine/Plugins/Experimental/MeshLODToolset/Source/MeshLODToolset/Public/Graphs/GenerateStaticMeshLODProcess.h b/Engine/Plugins/Experimental/MeshLODToolset/Source/MeshLODToolset/Public/Graphs/GenerateStaticMeshLODProcess.h
--- a/Engine/Plugins/Experimental/MeshLODToolset/Source/MeshLODToolset/Public/Graphs/GenerateStaticMeshLODProcess.h
+++ b/Engine/Plugins/Experimental/MeshLODToolset/Source/MeshLODToolset/Public/Graphs/GenerateStaticMeshLODProcess.h
@@ -27,9 +27,23 @@ UENUM()
 enum class EGenerateStaticMeshLODProcess_MeshGeneratorModes : uint8
 {
 	//~NOTE: must keep in sync with FGenerateMeshLODGraph::ECoreMeshGeneratorMode
+
+	/**
+	 * Generate a mesh using Marching Cubes to remesh the input shape.
+	 * Use Winding Numbers to determine what is inside
+	 */
 	Solidify = 0,
+
+	/**
+	 * Like Solidify, but Dilate and Contract the shape to delete small
+	 * negative features (sharp inner corners, small holes) before remeshing
+	 */
 	SolidifyAndClose = 1,
+
+	/** Generate a mesh by simplifying input Mesh attributes and filling small holes */
 	CleanAndSimplify = 2,
+
+	/** Generate a mesh by calculating the Convex Hull of the input shape */
 	ConvexHull = 3
 };
 
@@ -39,6 +53,7 @@ struct FGenerateStaticMeshLODProcessSettings
 {
 	GENERATED_BODY()
 
+	/** Method used to generate the initial mesh for AutoLOD processing */
 	UPROPERTY(EditAnywhere, Category = MeshGenerator, meta = (DisplayName="Mesh Generator"))
 	EGenerateStaticMeshLODProcess_MeshGeneratorModes MeshGenerator = EGenerateStaticMeshLODProcess_MeshGeneratorModes::SolidifyAndClose;
 
@@ -46,12 +61,12 @@ struct FGenerateStaticMeshLODProcessSettings
 
 	/** Target number of voxels along the maximum dimension for Solidify operation */
 	UPROPERTY(EditAnywhere, Category = Solidify, meta = (UIMin = "8", UIMax = "1024", ClampMin = "8", ClampMax = "1024", DisplayName="Voxel Resolution",
-	EditConditionHides, EditCondition = "MeshGenerator != EGenerateStaticMeshLODProcess_MeshGeneratorModes::ConvexHull"))
+	EditConditionHides, EditCondition = "MeshGenerator == EGenerateStaticMeshLODProcess_MeshGeneratorModes::Solidify || MeshGenerator == EGenerateStaticMeshLODProcess_MeshGeneratorModes::SolidifyAndClose"))
 	int SolidifyVoxelResolution = 128;
 
 	/** Winding number threshold to determine what is considered inside the mesh during Solidify */
 	UPROPERTY(EditAnywhere, Category = Solidify, AdvancedDisplay, meta = (UIMin = "0.1", UIMax = ".9", ClampMin = "-10", ClampMax = "10",
-	EditConditionHides, EditCondition = "MeshGenerator != EGenerateStaticMeshLODProcess_MeshGeneratorModes::ConvexHull"))
+	EditConditionHides, EditCondition = "MeshGenerator == EGenerateStaticMeshLODProcess_MeshGeneratorModes::Solidify || MeshGenerator == EGenerateStaticMeshLODProcess_MeshGeneratorModes::SolidifyAndClose"))
 	float WindingThreshold = 0.5f;
 
 
@@ -59,7 +74,7 @@ struct FGenerateStaticMeshLODProcessSettings
 
 	/** Offset distance in the Morpohological Closure operation */
 	UPROPERTY(EditAnywhere, Category = Morphology, meta = (UIMin = "0", UIMax = "100", ClampMin = "0", ClampMax = "1000", DisplayName = "Closure Distance",
-	EditConditionHides, EditCondition = "MeshGenerator != EGenerateStaticMeshLODProcess_MeshGeneratorModes::ConvexHull"))
+	EditConditionHides, EditCondition = "MeshGenerator == EGenerateStaticMeshLODProcess_MeshGeneratorModes::SolidifyAndClose"))
 	float ClosureDistance = 1.0f;
 
 
@@ -70,7 +85,7 @@ struct FGenerateStaticMeshLODProcessSettings
 	bool bPrefilterVertices = true;
 
 	/** Grid resolution (along the maximum-length axis) for subsampling before computing the convex hull */
-	UPROPERTY(EditAnywhere, Category = Collision, meta = (NoSpinbox = "true", EditConditionHides, EditCondition = "MeshGenerator == EGenerateStaticMeshLODProcess_MeshGeneratorModes::ConvexHull", ClampMin = "4", ClampMax = "100"))
+	UPROPERTY(EditAnywhere, Category = Collision, meta = (NoSpinbox = "true", EditConditionHides, EditCondition = "MeshGenerator == EGenerateStaticMeshLODProcess_MeshGeneratorModes::ConvexHull && bPrefilterVertices", ClampMin = "4", ClampMax = "100"))
 	int PrefilterGridResolution = 10;
 
 };
@@ -95,13 +110,24 @@ struct FGenerateStaticMeshLODProcess_PreprocessSettings
 
 	// Thicken settings
 
+	// TODO When ThickenWeightMapName is empty the behavior should be that the given ThickenAmount
+	// is applied to the entire mesh. This is not yet implemented so we disabled ThickenAmount to
+	// reduce confusion.
+
 	/** Weight map used during mesh thickening */
 	UPROPERTY(EditAnywhere, Category = Preprocessing, meta = (DisplayName = "Thicken Weight Map"))
 	FName ThickenWeightMapName = FName();
 
-	/** Amount to thicken the mesh prior to Solidifying. The thicken weight map values are multiplied by this value. */
-	UPROPERTY(EditAnywhere, Category = Preprocessing, meta = (UIMin = "0", UIMax = "100", ClampMin = "0", ClampMax = "1000"))
+	/**
+	 * Amount to thicken the mesh prior to Solidifying. The thicken weight map values are multiplied by this value.
+	 * Thickening is primarily intended to repair shape erosion that can occur when using the Solidify Mesh Generators.
+	 */
+	UPROPERTY(EditAnywhere, Category = Preprocessing, meta = (EditCondition = "bIsThickenAmountEnabled", HideEditConditionToggle, UIMin = "0", UIMax = "100", ClampMin = "0", ClampMax = "1000"))
 	float ThickenAmount = 0.0f;
+
+	// Transient property, not exposed to the user.
+	UPROPERTY(meta = (TransientToolProperty))
+	bool bIsThickenAmountEnabled = false;
 };
 
 
@@ -123,6 +149,7 @@ struct FGenerateStaticMeshLODProcess_SimplifySettings
 {
 	GENERATED_BODY()
 
+	/** Method used to simplify the mesh created in the Mesh Generation step */
 	UPROPERTY(EditAnywhere, Category = Simplify, meta = (DisplayName = "Simplification Target"))
 	EGenerateStaticMeshLODProcess_SimplifyMethod Method = EGenerateStaticMeshLODProcess_SimplifyMethod::GeometricTolerance;
 
@@ -258,6 +285,7 @@ struct FGenerateStaticMeshLODProcess_TextureSettings
 	UPROPERTY(EditAnywhere, Category = Baking, meta = (UIMin = "0", UIMax = "100", ClampMin = "0", ClampMax = "1000", DisplayName = "Bake Thickness"))
 	float BakeThickness = 5.0f;
 
+	/** Generate an atlassed texture where multiple textures are combined into a single output texture */
 	UPROPERTY(EditAnywhere, Category = Baking, meta = (DisplayName = "Combine Textures"))
 	bool bCombineTextures = true;
 };
diff --git a/Engine/Plugins/Experimental/MeshLODToolset/Source/MeshLODToolset/Public/Tools/GenerateStaticMeshLODAssetTool.h b/Engine/Plugins/Experimental/MeshLODToolset/Source/MeshLODToolset/Public/Tools/GenerateStaticMeshLODAssetTool.h
--- a/Engine/Plugins/Experimental/MeshLODToolset/Source/MeshLODToolset/Public/Tools/GenerateStaticMeshLODAssetTool.h
+++ b/Engine/Plugins/Experimental/MeshLODToolset/Source/MeshLODToolset/Public/Tools/GenerateStaticMeshLODAssetTool.h
@@ -72,7 +72,7 @@ public:
 
 	/** Suffix to append to newly-generated Asset (Meshes, Textures, Materials, etc) */
 	UPROPERTY(EditAnywhere, Category = "Output Options", meta = (TransientToolProperty))
-	FString GeneratedSuffix;
+	FString GeneratedSuffix = TEXT("_AutoLOD");
 
 	/** If false, then OutputMode will not be shown in DetailsView panels (otherwise no effect) */
 	UPROPERTY(meta = (TransientToolProperty))
@@ -105,14 +105,14 @@ public:
 	UPROPERTY(EditAnywhere, Category = Preset, meta = (DisplayName = "Settings Preset"))
 	TWeakObjectPtr<UStaticMeshLODGenerationSettings> Preset;
 
-	/** Save the current Tool settings to the Preset Asset */
+	/** Read the current Tool settings from the Preset Asset */
 	UFUNCTION(CallInEditor, Category = Preset, meta = (DisplayPriority = 0))
 	void ReadFromPreset()
 	{
 		PostAction(EGenerateLODAssetToolPresetAction::ReadFromPreset);
 	}
 
-	/** Read the current Tool settings from the Preset Asset */
+	/** Save the current Tool settings to the Preset Asset */
 	UFUNCTION(CallInEditor, Category = Preset, meta = (DisplayPriority = 1))
 	void WriteToPreset()
 	{
@@ -223,10 +223,10 @@ struct MESHLODTOOLSET_API FGenerateStaticMeshLOD_MaterialConfig
 {
 	GENERATED_BODY()
 
-	UPROPERTY(VisibleAnywhere, Category = Material)
+	UPROPERTY(VisibleAnywhere, Category = Material, meta = (TransientToolProperty, NoResetToDefault))
 	TObjectPtr<UMaterialInterface> Material = nullptr;
 
-	UPROPERTY(EditAnywhere, Category = Texture)
+	UPROPERTY(EditAnywhere, Category = Texture, meta = (TransientToolProperty, NoResetToDefault))
 	EGenerateStaticMeshLOD_BakeConstraint Constraint = EGenerateStaticMeshLOD_BakeConstraint::NoConstraint;
 
 	bool operator==(const FGenerateStaticMeshLOD_MaterialConfig& Other) const
@@ -242,13 +242,13 @@ class MESHLODTOOLSET_API UGenerateStaticMeshLODAssetToolTextureProperties : publ
 {
 	GENERATED_BODY()
 public:
-	UPROPERTY(EditAnywhere, Category = "Source Materials Configuration", meta = (TransientToolProperty))
+	UPROPERTY(EditAnywhere, Category = "Source Materials Configuration", meta = (TransientToolProperty, NoResetToDefault))
 	TArray<FGenerateStaticMeshLOD_MaterialConfig> Materials;
 
-	UPROPERTY(EditAnywhere, Category = "Source Textures Configuration", meta = (TransientToolProperty))
+	UPROPERTY(EditAnywhere, Category = "Source Textures Configuration", meta = (TransientToolProperty, NoResetToDefault))
 	TArray<FGenerateStaticMeshLOD_TextureConfig> Textures;
 
-	UPROPERTY(VisibleAnywhere, Category = "Baked Texture Previews")
+	UPROPERTY(VisibleAnywhere, Category = "Baked Texture Previews", meta = (TransientToolProperty, NoResetToDefault))
 	TArray<TObjectPtr<UTexture2D>> PreviewTextures;
 };
 
