ModelingTools: Prevent AutoLOD tool from running on dead StaticMesh assets.

diff --git a/Engine/Plugins/Experimental/MeshLODToolset/Source/MeshLODToolset/Private/Graphs/GenerateStaticMeshLODProcess.cpp b/Engine/Plugins/Experimental/MeshLODToolset/Source/MeshLODToolset/Private/Graphs/GenerateStaticMeshLODProcess.cpp
--- a/Engine/Plugins/Experimental/MeshLODToolset/Source/MeshLODToolset/Private/Graphs/GenerateStaticMeshLODProcess.cpp
+++ b/Engine/Plugins/Experimental/MeshLODToolset/Source/MeshLODToolset/Private/Graphs/GenerateStaticMeshLODProcess.cpp
@@ -1624,6 +1624,14 @@ void UGenerateStaticMeshLODProcess::WriteDerivedStaticMeshAsset()
 		GeneratedStaticMesh = Cast<UStaticMesh>(DupeAsset);
 	}
 
+	// DuplicateAsset can fail in the event of a dead package from an Asset Reload.
+	// This edge case is currently prevented by the GenerateStaticMeshLODAssetToolBuilder.
+	// [TODO] Handle a failed DuplicateAsset case by creating a new StaticMesh asset.
+	if (!ensure(GeneratedStaticMesh))
+	{
+		return;
+	}
+
 	// make sure transactional flag is on
 	GeneratedStaticMesh->SetFlags(RF_Transactional);
 	GeneratedStaticMesh->Modify();
diff --git a/Engine/Plugins/Experimental/MeshLODToolset/Source/MeshLODToolset/Private/Tools/GenerateStaticMeshLODAssetTool.cpp b/Engine/Plugins/Experimental/MeshLODToolset/Source/MeshLODToolset/Private/Tools/GenerateStaticMeshLODAssetTool.cpp
--- a/Engine/Plugins/Experimental/MeshLODToolset/Source/MeshLODToolset/Private/Tools/GenerateStaticMeshLODAssetTool.cpp
+++ b/Engine/Plugins/Experimental/MeshLODToolset/Source/MeshLODToolset/Private/Tools/GenerateStaticMeshLODAssetTool.cpp
@@ -21,7 +21,7 @@
 
 #include "Misc/Paths.h"
 #include "EditorAssetLibrary.h"
-
+#include "ObjectTools.h"
 
 #include "ModelingOperators.h"
 #include "MeshOpPreviewHelpers.h"
@@ -214,7 +214,25 @@ const FToolTargetTypeRequirements& UGenerateStaticMeshLODAssetToolBuilder::GetTa
 bool UGenerateStaticMeshLODAssetToolBuilder::CanBuildTool(const FToolBuilderState& SceneState) const
 {
 	// hack to make multi-tool look like single-tool
-	return (SceneState.TargetManager->CountSelectedAndTargetable(SceneState, GetTargetRequirements()) == 1);
+	const int NumTargets = SceneState.TargetManager->CountSelectedAndTargetable(SceneState, GetTargetRequirements());
+	if (NumTargets == 1)
+	{
+		// Ensure that the StaticMesh asset is still valid.
+		// [TODO] Remove this check once we properly handle a failed DuplicateAsset in UGenerateStaticMeshLODProcess::WriteDerivedStaticMeshAsset
+		bool bValidTarget = false;
+		SceneState.TargetManager->EnumerateSelectedAndTargetableComponents(SceneState, GetTargetRequirements(),
+			[&bValidTarget](UActorComponent* Component)
+			{
+				const UStaticMeshComponent* StaticMeshComponent = Cast<UStaticMeshComponent>(Component);
+				UStaticMesh* StaticMesh = StaticMeshComponent ? StaticMeshComponent->GetStaticMesh() : nullptr;
+				if (IsValid(StaticMesh))
+				{
+					bValidTarget = ObjectTools::IsObjectBrowsable(StaticMesh);
+				}
+			});
+		return bValidTarget;
+	}
+	return false;
 }
 
 UMultiSelectionMeshEditingTool* UGenerateStaticMeshLODAssetToolBuilder::CreateNewTool(const FToolBuilderState& SceneState) const
