Change LOD manager's "Show Seams" option to "Show Borders" and disable it when the default LOD is showing

diff --git a/Engine/Plugins/Experimental/MeshLODToolset/Source/MeshLODToolset/Private/Tools/LODManagerTool.cpp b/Engine/Plugins/Experimental/MeshLODToolset/Source/MeshLODToolset/Private/Tools/LODManagerTool.cpp
--- a/Engine/Plugins/Experimental/MeshLODToolset/Source/MeshLODToolset/Private/Tools/LODManagerTool.cpp
+++ b/Engine/Plugins/Experimental/MeshLODToolset/Source/MeshLODToolset/Private/Tools/LODManagerTool.cpp
@@ -89,6 +89,7 @@ void ULODManagerTool::Setup()
 	LODPreviewProperties = NewObject<ULODManagerPreviewLODProperties>(this);
 	AddToolPropertySource(LODPreviewProperties);
 	LODPreviewProperties->VisibleLOD = this->DefaultLODName;
+	LODPreviewProperties->bShowingDefaultLOD = true;
 	LODPreviewProperties->WatchProperty(LODPreviewProperties->VisibleLOD, [this](FString NewLOD) { bPreviewLODValid = false; });
 	LODPreviewProperties->WatchProperty(LODPreviewProperties->bShowSeams, [this](bool bNewValue) { bPreviewLODValid = false; });
 
@@ -386,6 +387,15 @@ void ULODManagerTool::UpdatePreviewLOD()
 	}
 	bPreviewLODValid = true;
 
+	auto SetShowingDefaultLOD = [this](bool bNewValue)
+	{
+		if (bNewValue != LODPreviewProperties->bShowingDefaultLOD)
+		{
+			LODPreviewProperties->bShowingDefaultLOD = bNewValue;
+			NotifyOfPropertyChangeByTool(LODPreviewProperties);
+		}
+	};
+
 	FString SelectedLOD = LODPreviewProperties->VisibleLOD;
 	const FLODName* FoundName = ActiveLODNames.Find(SelectedLOD);
 	if (SelectedLOD.IsEmpty() || FoundName == nullptr || FoundName->IsDefault() )
@@ -395,6 +405,7 @@ void ULODManagerTool::UpdatePreviewLOD()
 		ClearPreviewLines();
 		LODPreviewLines->SetAllVisible(false);
 		UE::ToolTarget::ShowSourceObject(Targets[0]);
+		SetShowingDefaultLOD(true);
 		return;
 	}
 
@@ -404,6 +415,8 @@ void ULODManagerTool::UpdatePreviewLOD()
 		CacheLODMesh(SelectedLOD, *FoundName);
 		Found = LODMeshCache.Find(SelectedLOD);
 	}
+
+	SetShowingDefaultLOD(Found == nullptr);
 	
 	if (Found)
 	{
diff --git a/Engine/Plugins/Experimental/MeshLODToolset/Source/MeshLODToolset/Public/Tools/LODManagerTool.h b/Engine/Plugins/Experimental/MeshLODToolset/Source/MeshLODToolset/Public/Tools/LODManagerTool.h
--- a/Engine/Plugins/Experimental/MeshLODToolset/Source/MeshLODToolset/Public/Tools/LODManagerTool.h
+++ b/Engine/Plugins/Experimental/MeshLODToolset/Source/MeshLODToolset/Public/Tools/LODManagerTool.h
@@ -89,8 +89,12 @@ public:
 	TArray<FString> LODNamesList;
 
 	/** Control whether mesh borders are displayed */
-	UPROPERTY(EditAnywhere, Category = LODPreview)
+	UPROPERTY(EditAnywhere, DisplayName = "Show Borders", Category = LODPreview, meta = (EditCondition = "!bShowingDefaultLOD", HideEditConditionToggle))
 	bool bShowSeams = true;
+
+	// Used to expose the VisibleLOD state to the above EditCondition
+	UPROPERTY(Transient, meta = (TransientToolProperty))
+	bool bShowingDefaultLOD = true;
 };
 
 
