Update LOD Manager's preview of nanite triangle percent to the newer/more-prominent Keep Triangle Percent. Also scale by 100 to make it match the UI display and add a tooltip.

diff --git a/Engine/Plugins/Experimental/MeshLODToolset/Source/MeshLODToolset/Private/Tools/LODManagerTool.cpp b/Engine/Plugins/Experimental/MeshLODToolset/Source/MeshLODToolset/Private/Tools/LODManagerTool.cpp
--- a/Engine/Plugins/Experimental/MeshLODToolset/Source/MeshLODToolset/Private/Tools/LODManagerTool.cpp
+++ b/Engine/Plugins/Experimental/MeshLODToolset/Source/MeshLODToolset/Private/Tools/LODManagerTool.cpp
@@ -223,7 +223,7 @@ void ULODManagerTool::UpdateLODInfo()
 	bLODInfoValid = true;
 
 	LODInfoProperties->bNaniteEnabled = StaticMesh->IsNaniteEnabled();
-	LODInfoProperties->PercentTriangles = StaticMesh->NaniteSettings.FallbackPercentTriangles;
+	LODInfoProperties->KeepTrianglePercent = StaticMesh->NaniteSettings.KeepPercentTriangles * 100;
 
 	TArray<FStaticMaterial> CurMaterialSet = StaticMesh->GetStaticMaterials();
 	LODInfoProperties->Materials.Reset();
diff --git a/Engine/Plugins/Experimental/MeshLODToolset/Source/MeshLODToolset/Public/Tools/LODManagerTool.h b/Engine/Plugins/Experimental/MeshLODToolset/Source/MeshLODToolset/Public/Tools/LODManagerTool.h
--- a/Engine/Plugins/Experimental/MeshLODToolset/Source/MeshLODToolset/Public/Tools/LODManagerTool.h
+++ b/Engine/Plugins/Experimental/MeshLODToolset/Source/MeshLODToolset/Public/Tools/LODManagerTool.h
@@ -63,8 +63,9 @@ public:
 	UPROPERTY(VisibleAnywhere, Category = Nanite, meta = (NoResetToDefault, DisplayName="Enabled"))
 	bool bNaniteEnabled = false;
 
+	// Percentage of triangles kept by Nanite
 	UPROPERTY(VisibleAnywhere, Category = Nanite, meta = (NoResetToDefault))
-	float PercentTriangles = 0.0f;
+	float KeepTrianglePercent = 0.0f;
 
 	UPROPERTY(VisibleAnywhere, Category = Materials, meta = (NoResetToDefault))
 	TArray<FStaticMaterial> Materials;
