Move MeshLODToolset into Editor plugins folder

diff --git a/Engine/Plugins/Editor/MeshLODToolset/MeshLODToolset.uplugin b/Engine/Plugins/Editor/MeshLODToolset/MeshLODToolset.uplugin
new file mode 100644
--- /dev/null
+++ b/Engine/Plugins/Editor/MeshLODToolset/MeshLODToolset.uplugin
@@ -0,0 +1,43 @@
+{
+	"FileVersion": 3,
+	"Version": 1,
+	"VersionName": "0.1",
+	"FriendlyName": "Mesh LOD Toolset",
+	"Description": "A set of modules implementing 3D mesh LOD creation",
+	"Category": "Other",
+	"CreatedBy": "Epic Games, Inc.",
+	"CreatedByURL": "https://epicgames.com",
+	"DocsURL": "",
+	"MarketplaceURL": "",
+	"SupportURL": "",
+	"CanContainContent": true,
+	"IsBetaVersion": true,
+	"Installed": false,
+	"Hidden": true,
+	"Modules": [
+		{
+			"Name": "MeshLODToolset",
+			"Type": "Editor",
+			"LoadingPhase": "Default"
+		},
+	],
+	"Plugins": [
+		{
+			"Name": "GeometryProcessing",
+			"Enabled": true
+		},
+		{
+			"Name": "GeometryFlow",
+			"Enabled": true
+		},
+		{
+			"Name": "MeshModelingToolsetExp",
+			"Enabled": true,
+			"TargetAllowList": [ "Editor" ]
+		},
+		{
+			"Name": "EditorScriptingUtilities",
+			"Enabled": true
+		}
+	]
+}
diff --git a/Engine/Plugins/Editor/MeshLODToolset/Source/MeshLODToolset/MeshLODToolset.Build.cs b/Engine/Plugins/Editor/MeshLODToolset/Source/MeshLODToolset/MeshLODToolset.Build.cs
new file mode 100644
--- /dev/null
+++ b/Engine/Plugins/Editor/MeshLODToolset/Source/MeshLODToolset/MeshLODToolset.Build.cs
@@ -0,0 +1,67 @@
+// Copyright Epic Games, Inc. All Rights Reserved.
+
+using UnrealBuildTool;
+
+public class MeshLODToolset : ModuleRules
+{
+	public MeshLODToolset(ReadOnlyTargetRules Target) : base(Target)
+	{
+		PCHUsage = ModuleRules.PCHUsageMode.UseExplicitOrSharedPCHs;
+
+		PublicIncludePaths.AddRange(
+			new string[]
+			{
+				// ... add public include paths required here ...
+			}
+		);
+
+		PrivateIncludePaths.AddRange(
+			new string[]
+			{
+				// ... add other private include paths required here ...
+			}
+		);
+
+		PublicDependencyModuleNames.AddRange(
+			new string[]
+			{
+				"Core",
+				"InteractiveToolsFramework",
+				"GeometryCore",
+				"DynamicMesh",
+				"MeshConversion",
+				"MeshConversionEngineTypes",
+				"MeshDescription",
+				"StaticMeshDescription",
+				"ModelingComponents",
+				"MeshModelingToolsExp",
+				"GeometryFlowCore",
+				"GeometryFlowMeshProcessing",
+				"GeometryFlowMeshProcessingEditor",
+				// ... add other public dependencies that you statically link with here ...
+			}
+		);
+
+		PrivateDependencyModuleNames.AddRange(
+			new string[]
+			{
+				"CoreUObject",
+				"Engine",
+				"RenderCore",
+				"UnrealEd",
+				"EditorScriptingUtilities",
+				"Slate",
+				"SlateCore",
+				"AssetDefinition"
+				// ... add private dependencies that you statically link with here ...
+			}
+		);
+
+		DynamicallyLoadedModuleNames.AddRange(
+			new string[]
+			{
+				// ... add any modules that your module loads dynamically here ...
+			}
+		);
+	}
+}
\ No newline at end of file
diff --git a/Engine/Plugins/Editor/MeshLODToolset/Source/MeshLODToolset/Private/Graphs/GenerateMeshLODGraph.cpp b/Engine/Plugins/Editor/MeshLODToolset/Source/MeshLODToolset/Private/Graphs/GenerateMeshLODGraph.cpp
new file mode 100644
--- /dev/null
+++ b/Engine/Plugins/Editor/MeshLODToolset/Source/MeshLODToolset/Private/Graphs/GenerateMeshLODGraph.cpp
@@ -0,0 +1,614 @@
+// Copyright Epic Games, Inc. All Rights Reserved.
+
+#include "Graphs/GenerateMeshLODGraph.h"
+
+#include "MeshLODToolsetModule.h"
+#include "GeometryFlowGraphUtil.h"
+#include "BaseNodes/TransferNode.h"
+#include "BaseNodes/SwitchNode.h"
+
+#include "MeshProcessingNodes/MeshDeleteTrianglesNode.h"
+#include "MeshProcessingNodes/CompactMeshNode.h"
+#include "MeshProcessingNodes/TransferMeshMaterialIDsNode.h"
+
+#include "MeshProcessingNodes/MeshTangentsNodes.h"
+
+#include "MeshProcessingNodes/MeshRepackUVsNode.h"
+
+#include "MeshBakingNodes/BakeMeshNormalMapNode.h"
+#include "MeshBakingNodes/BakeMeshMultiTextureNode.h"
+
+#include "MeshDecompositionNodes/MakeTriangleSetsNode.h"
+
+using namespace UE::Geometry;
+using namespace UE::GeometryFlow;
+
+
+typedef UE::GeometryFlow::TSwitchNode<FDynamicMesh3, 4, (int)EMeshProcessingDataTypes::DynamicMesh> FMeshGeneratorSwitchNode;
+
+
+void FGenerateMeshLODGraph::SetSourceMesh(const FDynamicMesh3& SourceMeshIn)
+{
+	UpdateSourceNodeValue<FDynamicMeshSourceNode>(*Graph, MeshSourceNode, SourceMeshIn);
+}
+
+
+void FGenerateMeshLODGraph::UpdatePreFilterSettings(const FMeshLODGraphPreFilterSettings& PreFilterSettings)
+{
+	UpdateSourceNodeValue<FNameSourceNode>(*Graph, FilterGroupsLayerNameNode, PreFilterSettings.FilterGroupLayerName);
+	CurrentPreFilterSettings = PreFilterSettings;
+}
+
+void FGenerateMeshLODGraph::UpdateCoreMeshGeneratorMode(ECoreMeshGeneratorMode NewMode)
+{
+	UpdateSwitchNodeInputIndex<FMeshGeneratorSwitchNode>(*Graph, MeshGeneratorSwitchNode, static_cast<int32>(NewMode) );
+	CurrentCoreMeshGeneratorMode = NewMode;
+}
+
+
+void FGenerateMeshLODGraph::UpdateSolidifySettings(const FMeshSolidifySettings& SolidifySettings)
+{
+	UpdateSettingsSourceNodeValue(*Graph, SolidifySettingsNode, SolidifySettings);
+	CurrentSolidifySettings = SolidifySettings;
+}
+
+void FGenerateMeshLODGraph::UpdateMorphologySettings(const FVoxClosureSettings& MorphologySettings)
+{
+	UpdateSettingsSourceNodeValue(*Graph, MorphologySettingsNode, MorphologySettings);
+	CurrentMorphologySettings = MorphologySettings;
+}
+
+void FGenerateMeshLODGraph::UpdateMeshCleaningSettings(const FMeshMakeCleanGeometrySettings& CleaningSettings)
+{
+	UpdateSettingsSourceNodeValue(*Graph, CleanMeshSettingsNode, CleaningSettings);
+	CurrentCleanMeshSettings = CleaningSettings;
+}
+
+void FGenerateMeshLODGraph::UpdateSimplifySettings(const FMeshSimplifySettings& SimplifySettings)
+{
+	UpdateSettingsSourceNodeValue(*Graph, SimplifySettingsNode, SimplifySettings);
+	CurrentSimplifySettings = SimplifySettings;
+}
+
+void FGenerateMeshLODGraph::UpdateGenerateConvexHullMeshSettings(const FGenerateConvexHullMeshSettings& GenerateConvexHullMeshSettings)
+{
+	UpdateSettingsSourceNodeValue(*Graph, GenerateConvexHullMeshSettingsNode, GenerateConvexHullMeshSettings);
+	CurrentGenerateConvexHullMeshSettings = GenerateConvexHullMeshSettings;
+}
+
+void FGenerateMeshLODGraph::UpdateNormalsSettings(const FMeshNormalsSettings& NormalsSettings)
+{
+	UpdateSettingsSourceNodeValue(*Graph, NormalsSettingsNode, NormalsSettings);
+	CurrentNormalsSettings = NormalsSettings;
+}
+
+void FGenerateMeshLODGraph::UpdateAutoUVSettings(const UE::GeometryFlow::FMeshAutoGenerateUVsSettings& AutoUVSettings)
+{
+	UpdateSettingsSourceNodeValue(*Graph, AutoUVSettingsNode, AutoUVSettings);
+	CurrentAutoUVSettings = AutoUVSettings;
+}
+
+void FGenerateMeshLODGraph::UpdateBakeCacheSettings(const UE::GeometryFlow::FMeshMakeBakingCacheSettings& BakeCacheSettings)
+{
+	UpdateSettingsSourceNodeValue(*Graph, BakeCacheSettingsNode, BakeCacheSettings);
+	CurrentBakeCacheSettings = BakeCacheSettings;
+}
+
+
+void FGenerateMeshLODGraph::UpdateGenerateSimpleCollisionSettings(const FGenerateSimpleCollisionSettings& GenSimpleCollisionSettings)
+{
+	UpdateSettingsSourceNodeValue(*Graph, GenerateSimpleCollisionSettingsNode, GenSimpleCollisionSettings);
+	CurrentGenerateSimpleCollisionSettings = GenSimpleCollisionSettings;
+}
+
+void FGenerateMeshLODGraph::UpdateThickenWeightMap(const TArray<float>& ThickenWeightMap)
+{
+	FWeightMap WeightMap;
+	WeightMap.Weights = ThickenWeightMap;
+	UpdateSourceNodeValue<FWeightMapSourceNode>(*Graph, ThickenWeightMapNode, WeightMap);
+}
+
+void FGenerateMeshLODGraph::UpdateThickenSettings(const UE::GeometryFlow::FMeshThickenSettings& ThickenSettings)
+{
+	UpdateSettingsSourceNodeValue(*Graph, ThickenSettingsNode, ThickenSettings);
+	CurrentThickenSettings = ThickenSettings;
+}
+
+void FGenerateMeshLODGraph::UpdateCollisionGroupLayerName(const FName& NewCollisionGroupLayerName)
+{
+	CollisionGroupLayerName = NewCollisionGroupLayerName;
+	UpdateSourceNodeValue<FNameSourceNode>(*Graph, GroupLayerNameNode, CollisionGroupLayerName);
+}
+
+
+void FGenerateMeshLODGraph::EvaluateResult(
+	FDynamicMesh3& ResultMesh,
+	FMeshTangentsd& ResultTangents,
+	FSimpleShapeSet3d& ResultCollision,
+	UE::GeometryFlow::FNormalMapImage& NormalMap,
+	TArray<TUniquePtr<UE::GeometryFlow::FTextureImage>>& TextureImages,
+	UE::GeometryFlow::FTextureImage& MultiTextureImage,
+	FProgressCancel* Progress)
+{
+	//FScopedDurationTimeLogger Timer(TEXT("FGenerateMeshLODGraph::EvaluateResult -- serial execution"));
+
+	if (Progress && Progress->Cancelled())
+	{
+		return;
+	}
+
+	//
+	// evaluate normal map
+	//
+
+	NormalMap = FNormalMapImage();
+	TUniquePtr<FEvaluationInfo> NormalMapEvalInfo = MakeUnique<FEvaluationInfo>();
+	NormalMapEvalInfo->Progress = Progress;
+	EGeometryFlowResult NormalMapEvalResult = Graph->EvaluateResult(BakeNormalMapNode, FBakeMeshNormalMapNode::OutParamNormalMap(),
+		NormalMap, (int)EMeshProcessingDataTypes::NormalMapImage, NormalMapEvalInfo, true);
+
+	if (Progress && Progress->Cancelled())
+	{
+		return;
+	}
+
+	ensure(NormalMapEvalResult == EGeometryFlowResult::Ok);
+	UE_LOG(LogMeshLODToolset, Display, TEXT("NormalMapPass - Evaluated %d Nodes, Recomputed %d"), NormalMapEvalInfo->NumEvaluations(), NormalMapEvalInfo->NumComputes());
+
+
+	//
+	// evaluate transferred textures
+	//
+
+	for (FBakeTextureGraphInfo& TexBakeStep : BakeTextureNodes)
+	{
+		TUniquePtr<UE::GeometryFlow::FTextureImage> NewImage = MakeUnique<UE::GeometryFlow::FTextureImage>();
+		TUniquePtr<FEvaluationInfo> TexBakeEvalInfo = MakeUnique<FEvaluationInfo>();
+		TexBakeEvalInfo->Progress = Progress;
+		EGeometryFlowResult TexBakeEvalResult = Graph->EvaluateResult(TexBakeStep.BakeNode, FBakeMeshTextureImageNode::OutParamTextureImage(),
+			*NewImage, (int)EMeshProcessingDataTypes::TextureImage, TexBakeEvalInfo, true);
+
+		if (Progress && Progress->Cancelled())
+		{
+			return;
+		}
+
+		TextureImages.Add(MoveTemp(NewImage));
+		ensure(TexBakeEvalResult == EGeometryFlowResult::Ok);
+
+		UE_LOG(LogMeshLODToolset, Display, TEXT("TextureBakePass %s - Evaluated %d Nodes, Recomputed %d"), *TexBakeStep.Identifier, TexBakeEvalInfo->NumEvaluations(), TexBakeEvalInfo->NumComputes());
+	}
+
+
+	// 
+	// evaluate multi texture bake
+	//
+	{
+		TUniquePtr<FEvaluationInfo> MultiTextureBakeEvalInfo = MakeUnique<FEvaluationInfo>();
+		MultiTextureBakeEvalInfo->Progress = Progress;
+		EGeometryFlowResult MultiTextureBakeEvalResult = Graph->EvaluateResult(BakeMultiTextureNode, FBakeMeshMultiTextureNode::OutParamTextureImage(),
+																			   MultiTextureImage, (int)EMeshProcessingDataTypes::TextureImage, MultiTextureBakeEvalInfo, true);
+
+		if (Progress && Progress->Cancelled())
+		{
+			return;
+		}
+
+		ensure(MultiTextureBakeEvalResult == EGeometryFlowResult::Ok);
+		UE_LOG(LogMeshLODToolset, Display, TEXT("MultiTextureBake - Evaluated %d Nodes, Recomputed %d"), MultiTextureBakeEvalInfo->NumEvaluations(), MultiTextureBakeEvalInfo->NumComputes());
+	}
+
+	// 
+	// evaluate tangents
+	//
+
+	bool bTakeResultTangents = false;
+	ResultTangents = FMeshTangentsd();
+
+	TUniquePtr<FEvaluationInfo> TangentsEvalInfo = MakeUnique<FEvaluationInfo>();
+	TangentsEvalInfo->Progress = Progress;
+	EGeometryFlowResult TangentsEvalResult = Graph->EvaluateResult(TangentsOutputNode, FMeshTangentsTransferNode::OutParamValue(),
+		ResultTangents, (int)EMeshProcessingDataTypes::MeshTangentSet, TangentsEvalInfo, bTakeResultTangents);
+
+	if (Progress && Progress->Cancelled())
+	{
+		return;
+	}
+
+	ensure(TangentsEvalResult == EGeometryFlowResult::Ok);
+	UE_LOG(LogMeshLODToolset, Display, TEXT("OutputTangentsPass - Evaluated %d Nodes, Recomputed %d"), TangentsEvalInfo->NumEvaluations(), TangentsEvalInfo->NumComputes());
+
+
+	//
+	// evaluate result mesh
+	// 
+
+
+	bool bTakeResultMesh = true;
+	ResultMesh.Clear();
+
+	TUniquePtr<FEvaluationInfo> MeshEvalInfo = MakeUnique<FEvaluationInfo>();
+	MeshEvalInfo->Progress = Progress;
+	EGeometryFlowResult EvalResult = Graph->EvaluateResult(MeshOutputNode, FDynamicMeshTransferNode::OutParamValue(),
+		ResultMesh, (int32)EMeshProcessingDataTypes::DynamicMesh, MeshEvalInfo, bTakeResultMesh);
+
+	if (Progress && Progress->Cancelled())
+	{
+		return;
+	}
+
+	ensure(EvalResult == EGeometryFlowResult::Ok);
+	UE_LOG(LogMeshLODToolset, Display, TEXT("OutputMeshPass - Evaluated %d Nodes, Recomputed %d"), MeshEvalInfo->NumEvaluations(), MeshEvalInfo->NumComputes());
+
+	//
+	// evaluate collision
+	//
+
+	bool bTakeResultCollision = false;
+	ResultCollision = FSimpleShapeSet3d();
+
+	TUniquePtr<FEvaluationInfo> CollisionEvalInfo = MakeUnique<FEvaluationInfo>();
+	CollisionEvalInfo->Progress = Progress;
+	EGeometryFlowResult CollisionEvalResult = Graph->EvaluateResult(CollisionOutputNode, FCollisionGeometryTransferNode::OutParamValue(),
+																	ResultCollision, FCollisionGeometry::DataTypeIdentifier, CollisionEvalInfo, bTakeResultCollision);
+
+	if (Progress && Progress->Cancelled())
+	{
+		return;
+	}
+
+	ensure(CollisionEvalResult == EGeometryFlowResult::Ok);
+	UE_LOG(LogMeshLODToolset, Display, TEXT("OutputCollisionPass - Evaluated %d Nodes, Recomputed %d"), CollisionEvalInfo->NumEvaluations(), CollisionEvalInfo->NumComputes());
+
+}
+
+
+
+static bool IsSingleMaterialMesh(const FDynamicMesh3& Mesh, int32& UniqueMaterialID)
+{
+	if (Mesh.HasAttributes() == false) return true;
+	if (Mesh.Attributes()->HasMaterialID() == false) return true;
+
+	const FDynamicMeshMaterialAttribute* MaterialIDs = Mesh.Attributes()->GetMaterialID();
+	UniqueMaterialID = 0;
+	bool bConstantInitialized = false;
+	for (int32 TriangleID : Mesh.TriangleIndicesItr())
+	{
+		int32 MaterialID = MaterialIDs->GetValue(TriangleID);
+		if (bConstantInitialized && MaterialID != UniqueMaterialID)
+		{
+			return false;
+		}
+		if (!bConstantInitialized)
+		{
+			bConstantInitialized = true;
+			UniqueMaterialID = MaterialID;
+		}
+	}
+	return true;
+}
+
+
+
+void FGenerateMeshLODGraph::BuildGraph(const FDynamicMesh3* SourceMeshHint)
+{
+	// precompute hints (async?)
+	int32 UniqueMaterialID = 0;
+	bool bIsSingleMaterialMesh = (SourceMeshHint != nullptr) ? IsSingleMaterialMesh(*SourceMeshHint, UniqueMaterialID) : false;
+
+	Graph = MakeUnique<FGraph>();
+
+	MeshSourceNode = Graph->AddNodeOfType<FDynamicMeshSourceNode>(TEXT("SourceMesh"));
+
+	// remove detail triangles
+
+	FGraph::FHandle FilterGroupsNode = Graph->AddNodeOfType<FIndexSetsSourceNode>(TEXT("FilterGroups"));
+	FilterGroupsLayerNameNode = Graph->AddNodeOfType<FNameSourceNode>(TEXT("FilterGroupsLayerNameSource"));
+
+	FGraph::FHandle MakeFilterTriangleSetsNode = Graph->AddNodeOfType<FMakeTriangleSetsFromGroupsNode>(TEXT("MakeFilterTriangles"));
+	Graph->InferConnection(MeshSourceNode, MakeFilterTriangleSetsNode);
+	Graph->InferConnection(FilterGroupsNode, MakeFilterTriangleSetsNode);
+	Graph->InferConnection(FilterGroupsLayerNameNode, MakeFilterTriangleSetsNode);
+
+	FilterTrianglesNode = Graph->AddNodeOfType<FMeshDeleteTrianglesNode>(TEXT("FilterMesh"));
+	Graph->InferConnection(MeshSourceNode, FilterTrianglesNode);
+	Graph->InferConnection(MakeFilterTriangleSetsNode, FilterTrianglesNode);
+
+	// generating low-poly mesh
+
+	// optionally thicken some parts of the mesh before solidifying
+	ThickenNode = Graph->AddNodeOfType<FMeshThickenNode>(TEXT("Thicken"));
+	ThickenWeightMapNode = Graph->AddNodeOfType<FWeightMapSourceNode>(TEXT("ThickenWeightMapNode"));
+	ThickenSettingsNode = Graph->AddNodeOfType<FThickenSettingsSourceNode>(TEXT("ThickenSettingsSource"));
+
+	Graph->InferConnection(ThickenWeightMapNode, ThickenNode);
+	Graph->InferConnection(ThickenSettingsNode, ThickenNode);
+	Graph->InferConnection(FilterTrianglesNode, ThickenNode);
+
+	SolidifyNode = Graph->AddNodeOfType<FSolidifyMeshNode>(TEXT("Solidify"));
+	Graph->InferConnection(ThickenNode, SolidifyNode);
+	SolidifySettingsNode = Graph->AddNodeOfType<FSolidifySettingsSourceNode>(TEXT("SolidifySettings"));
+	Graph->InferConnection(SolidifySettingsNode, SolidifyNode);
+
+	MorphologyNode = Graph->AddNodeOfType<FVoxClosureMeshNode>(TEXT("Closure"));
+	Graph->InferConnection(SolidifyNode, MorphologyNode);
+	MorphologySettingsNode = Graph->AddNodeOfType<FVoxClosureSettingsSourceNode>(TEXT("ClosureSettings"));
+	Graph->InferConnection(MorphologySettingsNode, MorphologyNode);
+
+	CleanMeshNode = Graph->AddNodeOfType<FMeshMakeCleanGeometryNode>(TEXT("Cleaner"));
+	Graph->InferConnection(ThickenNode, CleanMeshNode);
+	CleanMeshSettingsNode = Graph->AddNodeOfType<FMeshMakeCleanGeometrySettingsSourceNode>(TEXT("CleanerSettings"));
+	Graph->InferConnection(CleanMeshSettingsNode, CleanMeshNode);
+
+	GenerateConvexHullMeshNode = Graph->AddNodeOfType<FGenerateConvexHullMeshNode>(TEXT("ConvexHullMesh"));
+	Graph->InferConnection(FilterTrianglesNode, GenerateConvexHullMeshNode);
+	GenerateConvexHullMeshSettingsNode = Graph->AddNodeOfType<FGenerateConvexHullMeshSettingsSourceNode>(TEXT("ConvexHullMeshSettings"));
+	Graph->InferConnection(GenerateConvexHullMeshSettingsNode, GenerateConvexHullMeshNode);
+
+
+	// integer param that controls which mesh generator will be used
+	MeshGeneratorSwitchNode = Graph->AddNodeOfType<FMeshGeneratorSwitchNode>(TEXT("MeshGeneratorSwitch"));
+	Graph->AddConnection(SolidifyNode, FSolidifyMeshNode::OutParamResultMesh(), MeshGeneratorSwitchNode, FMeshGeneratorSwitchNode::InParamValue(0));
+	Graph->AddConnection(MorphologyNode, FVoxClosureMeshNode::OutParamResultMesh(), MeshGeneratorSwitchNode, FMeshGeneratorSwitchNode::InParamValue(1));
+	Graph->AddConnection(CleanMeshNode, FMeshMakeCleanGeometryNode::OutParamResultMesh(), MeshGeneratorSwitchNode, FMeshGeneratorSwitchNode::InParamValue(2));
+	Graph->AddConnection(GenerateConvexHullMeshNode, FGenerateConvexHullMeshNode::OutParamResultMesh(), MeshGeneratorSwitchNode, FMeshGeneratorSwitchNode::InParamValue(3));
+
+	// helper to separate next block from changes above
+	FGraph::FHandle FinalMeshOutputNode = MeshGeneratorSwitchNode;
+
+	// skip MaterialID projection/transfer if the input mesh only had one material
+	FGraph::FHandle SimplifyPhaseInputNode = FinalMeshOutputNode;
+	if (bIsSingleMaterialMesh == false)
+	{
+		FGraph::FHandle MatIDTransferNode = Graph->AddNodeOfType<FTransferMeshMaterialIDsNode>(TEXT("TransferMaterialIDs"));
+		Graph->AddConnection(MeshSourceNode, FDynamicMeshSourceNode::OutParamValue(), MatIDTransferNode, FTransferMeshMaterialIDsNode::InParamMaterialSourceMesh());
+		Graph->InferConnection(FinalMeshOutputNode, MatIDTransferNode);
+		SimplifyPhaseInputNode = MatIDTransferNode;
+	}
+	else
+	{
+		UE_LOG(LogGeometry, Warning, TEXT("AutoLOD: applying single-material optimizations"));
+	}
+
+	// need to compute valid normals before Simplify
+	//FGraph::FHandle PerVertexNormalsNode = Graph->AddNodeOfType<FComputeMeshPerVertexOverlayNormalsNode>(TEXT("PerVertexOverlayNormals"));
+	FGraph::FHandle PerVertexNormalsNode = Graph->AddNodeOfType<FComputeMeshPerVertexNormalsNode>(TEXT("PerVertexNormals"));
+	Graph->InferConnection(SimplifyPhaseInputNode, PerVertexNormalsNode);
+
+	SimplifyNode = Graph->AddNodeOfType<FSimplifyMeshNode>(TEXT("Simplify"));
+	Graph->InferConnection(PerVertexNormalsNode, SimplifyNode);
+	SimplifySettingsNode = Graph->AddNodeOfType<FSimplifySettingsSourceNode>(TEXT("SimplifySettings"));
+	Graph->InferConnection(SimplifySettingsNode, SimplifyNode);
+
+	FGraph::FHandle CompactNode = Graph->AddNodeOfType<FCompactMeshNode>(TEXT("Compact"));
+	Graph->InferConnection(SimplifyNode, CompactNode);
+
+	NormalsNode = Graph->AddNodeOfType<FComputeMeshNormalsNode>(TEXT("Normals"));
+	Graph->InferConnection(CompactNode, NormalsNode);
+	NormalsSettingsNode = Graph->AddNodeOfType<FNormalsSettingsSourceNode>(TEXT("NormalsSettings"));
+	Graph->InferConnection(NormalsSettingsNode, NormalsNode);
+
+	// computing UVs
+
+	AutoUVNode = Graph->AddNodeOfType<FMeshAutoGenerateUVsNode>(TEXT("AutoUV"));
+	Graph->InferConnection(NormalsNode, AutoUVNode);
+	AutoUVSettingsNode = Graph->AddNodeOfType<FMeshAutoGenerateUVsSettingsSourceNode>(TEXT("AutoUVSettings"));
+	Graph->InferConnection(AutoUVSettingsNode, AutoUVNode);
+
+	// disable RecomputeUVNode for now
+	//RecomputeUVNode = Graph->AddNodeOfType<FMeshRecalculateUVsNode>(TEXT("RecalcUV"));
+	//Graph->InferConnection(AutoUVNode, RecomputeUVNode);
+	//RecomputeUVSettingsNode = Graph->AddNodeOfType<FMeshRecalculateUVsSettingsSourceNode>(TEXT("RecalcUVSettings"));
+	//Graph->InferConnection(RecomputeUVSettingsNode, RecomputeUVNode);
+
+	RepackUVNode = Graph->AddNodeOfType<FMeshRepackUVsNode>(TEXT("RepackUV"));
+	//Graph->InferConnection(RecomputeUVNode, RepackUVNode);
+	Graph->InferConnection(AutoUVNode, RepackUVNode);
+	RepackUVSettingsNode = Graph->AddNodeOfType<FMeshRepackUVsSettingsSourceNode>(TEXT("RepackUVSettings"));
+	Graph->InferConnection(RepackUVSettingsNode, RepackUVNode);
+
+
+	// final mesh output
+
+	MeshOutputNode = Graph->AddNodeOfType<FDynamicMeshTransferNode>(TEXT("OutputMesh"));
+	Graph->InferConnection(RepackUVNode, MeshOutputNode);
+
+
+	// create tangents
+
+	TangentsNode = Graph->AddNodeOfType<FComputeMeshTangentsNode>(TEXT("Tangents"));
+	Graph->InferConnection(RepackUVNode, TangentsNode);
+	TangentsSettingsNode = Graph->AddNodeOfType<FTangentsSettingsSourceNode>(TEXT("TangentsSettings"));
+	Graph->InferConnection(TangentsSettingsNode, TangentsNode);
+
+	// tangents output
+	TangentsOutputNode = Graph->AddNodeOfType<FMeshTangentsTransferNode>(TEXT("OutputTangents"));
+	Graph->InferConnection(TangentsNode, TangentsOutputNode);
+
+	// create bake cache
+
+	BakeCacheNode = Graph->AddNodeOfType<FMakeMeshBakingCacheNode>(TEXT("MakeBakeCache"));
+	Graph->AddConnection(MeshSourceNode, FDynamicMeshSourceNode::OutParamValue(), BakeCacheNode, FMakeMeshBakingCacheNode::InParamDetailMesh());
+	Graph->AddConnection(RepackUVNode, FMeshRepackUVsNode::OutParamResultMesh(), BakeCacheNode, FMakeMeshBakingCacheNode::InParamTargetMesh());
+	BakeCacheSettingsNode = Graph->AddNodeOfType<FMeshMakeBakingCacheSettingsSourceNode>(TEXT("BakeCacheSettings"));
+	Graph->InferConnection(BakeCacheSettingsNode, BakeCacheNode);
+
+	// normal map baker
+
+	BakeNormalMapNode = Graph->AddNodeOfType<FBakeMeshNormalMapNode>(TEXT("BakeNormalMap"));
+	Graph->InferConnection(BakeCacheNode, BakeNormalMapNode);
+	Graph->InferConnection(TangentsNode, BakeNormalMapNode);
+	BakeNormalMapSettingsNode = Graph->AddNodeOfType<FBakeMeshNormalMapSettingsSourceNode>(TEXT("BakeNormalMapSettings"));
+	Graph->InferConnection(BakeNormalMapSettingsNode, BakeNormalMapNode);
+
+
+	// collision generation
+
+	FGraph::FHandle IgnoreGroupsForCollisionNode = Graph->AddNodeOfType<FIndexSetsSourceNode>(TEXT("CollisionIgnoreGroups"));
+
+	//DecomposeMeshForCollisionNode = Graph->AddNodeOfType<FMakeTriangleSetsFromMeshNode>(TEXT("Decompose"));
+	DecomposeMeshForCollisionNode = Graph->AddNodeOfType<FMakeTriangleSetsFromGroupsNode>(TEXT("Decompose"));
+	Graph->InferConnection(FilterTrianglesNode, DecomposeMeshForCollisionNode);
+	Graph->InferConnection(IgnoreGroupsForCollisionNode, DecomposeMeshForCollisionNode);
+
+	GroupLayerNameNode = Graph->AddNodeOfType<FNameSourceNode>(TEXT("GroupLayerNameNode"));
+	Graph->AddConnection(GroupLayerNameNode, FNameSourceNode::OutParamValue(), 
+						 DecomposeMeshForCollisionNode, FMakeTriangleSetsFromGroupsNode::InParamGroupLayer());
+
+	GenerateSimpleCollisionNode = Graph->AddNodeOfType<FGenerateSimpleCollisionNode>(TEXT("GenerateSimpleCollision"));
+	Graph->InferConnection(FilterTrianglesNode, GenerateSimpleCollisionNode);
+	Graph->InferConnection(DecomposeMeshForCollisionNode, GenerateSimpleCollisionNode);
+	GenerateSimpleCollisionSettingsNode = Graph->AddNodeOfType<FGenerateSimpleCollisionSettingsSourceNode>(TEXT("GenerateSimpleCollisionSettings"));
+	Graph->InferConnection(GenerateSimpleCollisionSettingsNode, GenerateSimpleCollisionNode);
+
+	// final collision output
+
+	CollisionOutputNode = Graph->AddNodeOfType<FCollisionGeometryTransferNode>(TEXT("OutputCollision"));
+	Graph->InferConnection(GenerateSimpleCollisionNode, CollisionOutputNode);
+
+
+	//
+	// parameters
+	//
+
+	CurrentCoreMeshGeneratorMode = ECoreMeshGeneratorMode::SolidifyAndClose;
+	UpdateSwitchNodeInputIndex<FMeshGeneratorSwitchNode>(*Graph, MeshGeneratorSwitchNode, static_cast<int32>(CurrentCoreMeshGeneratorMode) );
+
+	FIndexSets IgnoreGroupsForDelete;
+	IgnoreGroupsForDelete.AppendSet({ 0 });
+	UpdateSettingsSourceNodeValue(*Graph, FilterGroupsNode, IgnoreGroupsForDelete);
+
+	FMeshLODGraphPreFilterSettings PreFilterSettings;
+	PreFilterSettings.FilterGroupLayerName = FName();
+	UpdatePreFilterSettings(PreFilterSettings);
+
+	UE::GeometryFlow::FMeshThickenSettings MeshThickenSettings;
+	MeshThickenSettings.ThickenAmount = 10.0f;
+	UpdateThickenSettings(MeshThickenSettings);
+
+	FMeshSolidifySettings SolidifySettings;
+	SolidifySettings.VoxelResolution = 128;
+	UpdateSolidifySettings(SolidifySettings);
+
+	FVoxClosureSettings MorphologySettings;
+	MorphologySettings.VoxelResolution = SolidifySettings.VoxelResolution;
+	MorphologySettings.Distance = 5.0;
+	UpdateMorphologySettings(MorphologySettings);
+
+	FMeshMakeCleanGeometrySettings CleanerSettings;
+	CleanerSettings.bClearMaterialIDs = bIsSingleMaterialMesh;
+	CleanerSettings.bOutputOverlayVertexNormals = false;
+	UpdateMeshCleaningSettings(CleanerSettings);
+
+	// TODO: can use simpler settings if there is only one material...
+	FMeshSimplifySettings SimplifySettings;
+	SimplifySettings.bDiscardAttributes = false;
+	SimplifySettings.SimplifyType = EMeshSimplifyType::AttributeAware;
+	//SimplifySettings.TargetType = EMeshSimplifyTargetType::TriangleCount;
+	//SimplifySettings.TargetCount = 500;
+	SimplifySettings.TargetType = EMeshSimplifyTargetType::GeometricDeviation;
+	SimplifySettings.GeometricTolerance = 0.5;
+	SimplifySettings.TargetCount = 500;
+	SimplifySettings.MaterialBorderConstraints = EEdgeRefineFlags::NoFlip;
+	UpdateSimplifySettings(SimplifySettings);
+
+	FMeshNormalsSettings NormalsSettings;
+	NormalsSettings.NormalsType = EComputeNormalsType::FromFaceAngleThreshold;
+	NormalsSettings.AngleThresholdDeg = 60.0;
+	UpdateNormalsSettings(NormalsSettings);
+
+	FMeshAutoGenerateUVsSettings AutoUVSettings;
+	AutoUVSettings.Method = EAutoUVMethod::PatchBuilder;
+	AutoUVSettings.UVAtlasNumCharts = 20;
+	AutoUVSettings.UVAtlasStretch = 0.1;
+	AutoUVSettings.XAtlasMaxIterations = 1;
+	AutoUVSettings.NumInitialPatches = 100;
+	AutoUVSettings.CurvatureAlignment = 1.0;
+	AutoUVSettings.MergingThreshold = 1.5;
+	AutoUVSettings.MaxAngleDeviationDeg = 45.0;
+	AutoUVSettings.SmoothingSteps = 5;
+	AutoUVSettings.SmoothingAlpha = 0.25;
+	AutoUVSettings.bAutoPack = false;
+	AutoUVSettings.PackingTargetWidth = 512;
+	UpdateAutoUVSettings(AutoUVSettings);
+
+	//FMeshRecalculateUVsSettings RecomputeUVSettings;
+	//UpdateSettingsSourceNodeValue(*Graph, RecomputeUVSettingsNode, RecomputeUVSettings);
+
+	FMeshRepackUVsSettings RepackUVSettings;
+	UpdateSettingsSourceNodeValue(*Graph, RepackUVSettingsNode, RepackUVSettings);
+
+
+	FMeshTangentsSettings TangentsSettings;
+	UpdateSettingsSourceNodeValue(*Graph, TangentsSettingsNode, TangentsSettings);
+
+
+	FMeshMakeBakingCacheSettings BakeCacheSettings;
+	BakeCacheSettings.Dimensions = FImageDimensions(1024, 1024);
+	BakeCacheSettings.Thickness = 5.0;
+	UpdateBakeCacheSettings(BakeCacheSettings);
+
+
+	FBakeMeshNormalMapSettings NormalMapSettings;
+	UpdateSettingsSourceNodeValue(*Graph, BakeNormalMapSettingsNode, NormalMapSettings);
+
+	FIndexSets IgnoreGroupsForCollision;
+	IgnoreGroupsForCollision.AppendSet({ 0 });
+	UpdateSettingsSourceNodeValue(*Graph, IgnoreGroupsForCollisionNode, IgnoreGroupsForCollision);
+
+	UpdateCollisionGroupLayerName(CollisionGroupLayerName);
+
+	FGenerateSimpleCollisionSettings GenSimpleCollisionSettings;
+	GenSimpleCollisionSettings.Type = UE::GeometryFlow::ESimpleCollisionGeometryType::AlignedBoxes;
+	UpdateGenerateSimpleCollisionSettings(GenSimpleCollisionSettings);
+
+	TArray<float> Weights;
+	UpdateThickenWeightMap(Weights);
+}
+
+
+
+
+int32 FGenerateMeshLODGraph::AppendTextureBakeNode(const TImageBuilder<FVector4f>& SourceImage, const FString& Identifier)
+{
+	FBakeTextureGraphInfo NewNode;
+	NewNode.Index = BakeTextureNodes.Num();
+	NewNode.Identifier = Identifier;
+
+	// add source node
+	NewNode.TexSourceNode = Graph->AddNodeOfType<FTextureImageSourceNode>(FString::Printf(TEXT("TextureSource%d_%s"), NewNode.Index, *NewNode.Identifier));
+
+	// texture baker
+	NewNode.BakeNode = Graph->AddNodeOfType<FBakeMeshTextureImageNode>(FString::Printf(TEXT("BakeTexImage%d_%s"), NewNode.Index, *NewNode.Identifier));
+	ensure(Graph->InferConnection(BakeCacheNode, NewNode.BakeNode) == EGeometryFlowResult::Ok);
+	ensure(Graph->InferConnection(NewNode.TexSourceNode, NewNode.BakeNode) == EGeometryFlowResult::Ok);
+	FGraph::FHandle BakeTextureImageSettingsNode = Graph->AddNodeOfType<FBakeMeshTextureImageSettingsSourceNode>(TEXT("BakeTextureImageSettings"));
+	ensure(Graph->InferConnection(BakeTextureImageSettingsNode, NewNode.BakeNode) == EGeometryFlowResult::Ok);
+
+	FTextureImage InputTexImage;
+	InputTexImage.Image = SourceImage;
+	UpdateSourceNodeValue<FTextureImageSourceNode>(*Graph, NewNode.TexSourceNode, InputTexImage);
+
+	BakeTextureNodes.Add(NewNode);
+
+	return NewNode.Index;
+}
+
+
+void FGenerateMeshLODGraph::AppendMultiTextureBakeNode(const TMap<int32, TSafeSharedPtr<TImageBuilder<FVector4f>>>& SourceMaterialImages)
+{
+	// add source node
+	MaterialIDTextureSourceNode = Graph->AddNodeOfType<FMaterialIDToTextureMapSourceNode>(FString::Printf(TEXT("MaterialIDTextureSource")));
+
+	// texture baker
+	BakeMultiTextureNode = Graph->AddNodeOfType<FBakeMeshMultiTextureNode>(FString::Printf(TEXT("BakeMultiTexture")));
+	ensure(Graph->InferConnection(BakeCacheNode, BakeMultiTextureNode) == EGeometryFlowResult::Ok);
+	ensure(Graph->InferConnection(MaterialIDTextureSourceNode, BakeMultiTextureNode) == EGeometryFlowResult::Ok);
+	BakeMultiTextureSettingsNode = Graph->AddNodeOfType<FBakeMeshMultiTextureSettingsSourceNode>(TEXT("BakeMultiTextureSettings"));
+	ensure(Graph->InferConnection(BakeMultiTextureSettingsNode, BakeMultiTextureNode) == EGeometryFlowResult::Ok);
+
+	FMaterialIDToTextureMap InputMap;
+	InputMap.MaterialIDTextureMap = SourceMaterialImages;
+
+	UpdateSourceNodeValue<FMaterialIDToTextureMapSourceNode>(*Graph, MaterialIDTextureSourceNode, InputMap);
+}
diff --git a/Engine/Plugins/Editor/MeshLODToolset/Source/MeshLODToolset/Private/Graphs/GenerateStaticMeshLODProcess.cpp b/Engine/Plugins/Editor/MeshLODToolset/Source/MeshLODToolset/Private/Graphs/GenerateStaticMeshLODProcess.cpp
new file mode 100644
--- /dev/null
+++ b/Engine/Plugins/Editor/MeshLODToolset/Source/MeshLODToolset/Private/Graphs/GenerateStaticMeshLODProcess.cpp
@@ -0,0 +1,1894 @@
+// Copyright Epic Games, Inc. All Rights Reserved.
+
+#include "Graphs/GenerateStaticMeshLODProcess.h"
+
+#include "BodySetupEnums.h"
+#include "MeshLODToolsetModule.h"
+
+#include "Async/Async.h"
+
+#include "Engine/StaticMeshSourceData.h"
+#include "MeshDescriptionToDynamicMesh.h"
+#include "DynamicMeshToMeshDescription.h"
+
+#include "AssetUtils/Texture2DUtil.h"
+#include "AssetUtils/Texture2DBuilder.h"
+#include "AssetUtils/MeshDescriptionUtil.h"
+
+#include "Engine/Texture2D.h"
+#include "Physics/PhysicsDataCollection.h"
+
+#include "Graphs/GenerateMeshLODGraph.h"
+#include "Misc/Paths.h"
+#include "EditorAssetLibrary.h"
+#include "AssetToolsModule.h"
+#include "AssetRegistry/AssetRegistryModule.h"
+#include "IAssetTools.h"
+#include "Misc/PackageName.h"
+#include "UObject/MetaData.h"
+
+#include "Editor.h"
+
+#include "Materials/MaterialInstanceConstant.h"
+#include "Materials/MaterialExpressionTextureSampleParameter.h"
+#include "Factories/MaterialInstanceConstantFactoryNew.h"
+#include "Materials/MaterialInstanceDynamic.h"
+#include "MaterialGraph/MaterialGraphNode.h"
+#include "WeightMapUtil.h"
+
+
+#include "PhysicsEngine/BodySetup.h"
+#include "RenderingThread.h"
+#include "StaticMeshAttributes.h"
+
+#include UE_INLINE_GENERATED_CPP_BY_NAME(GenerateStaticMeshLODProcess)
+
+#define LOCTEXT_NAMESPACE "UGenerateStaticMeshLODProcess"
+
+using namespace UE::Geometry;
+using namespace UE::GeometryFlow;
+
+namespace
+{
+#if WITH_EDITOR
+	static EAsyncExecution GenerateSMLODAsyncExecTarget = EAsyncExecution::LargeThreadPool;
+#else
+	static EAsyncExecution GenerateSMLODAsyncExecTarget = EAsyncExecution::ThreadPool;
+#endif
+}
+
+
+
+struct FReadTextureJob
+{
+	int32 MatIndex;
+	int32 TexIndex;
+};
+
+
+namespace GenerateStaticMeshLODProcessHelpers
+{
+	// Given "xxxxx" returns "xxxxx_1"
+	// Given "xxxxx_1" returns "xxxxx_2"
+	// etc.
+	// If anything goes wrong, just add "_1" to the end
+	void AppendOrIncrementSuffix(FString& S)
+	{
+		TArray<FString> Substrings;
+		S.ParseIntoArray(Substrings, TEXT("_"));
+		if ( Substrings.Num() <= 1 ) 
+		{
+			S += TEXT("_1");
+		}
+		else
+		{
+			FString LastSubstring = Substrings.Last();
+			int32 Num;
+			bool bParsed = LexTryParseString(Num, *LastSubstring);
+			if (bParsed)
+			{
+				++Num;
+				Substrings.RemoveAt(Substrings.Num() - 1);
+				S = FString::Join(Substrings, TEXT("_"));
+				S += TEXT("_") + FString::FromInt(Num);
+			}
+			else
+			{
+				S += TEXT("_1");
+			}
+		}
+	}
+
+	TArray<int32> FindUnreferencedMaterials(UStaticMesh* StaticMesh, const FMeshDescription* MeshDescription)
+	{
+		const TArray<FStaticMaterial>& MaterialSet = StaticMesh->GetStaticMaterials();
+		const int32 NumMaterials = MaterialSet.Num();
+
+		auto IsValidMaterial = [&MaterialSet](int32 MaterialID) {
+			return MaterialSet[MaterialID].MaterialInterface != nullptr;
+		};
+
+		// Initially flag only valid materials as potentially unused.
+		TArray<bool> MatUnusedFlags;
+		MatUnusedFlags.SetNum(NumMaterials);
+		int32 NumMatUnused = 0;
+		for (int32 MaterialID = 0; MaterialID < NumMaterials; ++MaterialID)
+		{
+			MatUnusedFlags[MaterialID] = IsValidMaterial(MaterialID);
+			NumMatUnused += MatUnusedFlags[MaterialID];
+		}
+
+		TMap<FPolygonGroupID, int32> PolygonGroupToMaterialIndex;
+		const FStaticMeshConstAttributes Attributes(*MeshDescription);
+		TPolygonGroupAttributesConstRef<FName> PolygonGroupImportedMaterialSlotNames = Attributes.GetPolygonGroupMaterialSlotNames();
+
+		for (FPolygonGroupID PolygonGroupID : MeshDescription->PolygonGroups().GetElementIDs())
+		{
+			int32 MaterialIndex = StaticMesh->GetMaterialIndexFromImportedMaterialSlotName(PolygonGroupImportedMaterialSlotNames[PolygonGroupID]);
+			if (MaterialIndex == INDEX_NONE)
+			{
+				MaterialIndex = PolygonGroupID.GetValue();
+			}
+			PolygonGroupToMaterialIndex.Add(PolygonGroupID, MaterialIndex);
+		}
+
+		for (const FTriangleID TriangleID : MeshDescription->Triangles().GetElementIDs())
+		{
+			const FPolygonGroupID PolygonGroupID = MeshDescription->GetTrianglePolygonGroup(TriangleID);
+			const int32 MaterialIndex = PolygonGroupToMaterialIndex[PolygonGroupID];
+			bool& bMatUnusedFlag = MatUnusedFlags[MaterialIndex];
+			NumMatUnused -= static_cast<int32>(bMatUnusedFlag);
+			bMatUnusedFlag = false;
+			if (NumMatUnused == 0)
+			{
+				break;
+			}
+		}
+
+		TArray<int32> UnreferencedMaterials;
+		UnreferencedMaterials.Reserve(NumMatUnused);
+		for (int32 MaterialID = 0; MaterialID < NumMaterials; ++MaterialID)
+		{
+			if (MatUnusedFlags[MaterialID])
+			{
+				UnreferencedMaterials.Emplace(MaterialID);
+			}
+		}
+		return UnreferencedMaterials;
+	}
+}
+
+bool UGenerateStaticMeshLODProcess::Initialize(UStaticMesh* StaticMeshIn, FProgressCancel* Progress)
+{
+	if (!ensure(StaticMeshIn)) return false;
+	if (!ensure(StaticMeshIn->GetNumSourceModels() > 0)) return false;
+
+	// make sure we are not in rendering
+	FlushRenderingCommands();
+
+	SourceStaticMesh = StaticMeshIn;
+
+	bUsingHiResSource = SourceStaticMesh->IsHiResMeshDescriptionValid();
+	const FMeshDescription* UseSourceMeshDescription =
+		(bUsingHiResSource) ? SourceStaticMesh->GetHiResMeshDescription() : SourceStaticMesh->GetMeshDescription(0);
+
+	if (!UseSourceMeshDescription)
+	{
+		return false;
+	}
+
+	SourceMeshDescription = MakeShared<FMeshDescription>();
+	*SourceMeshDescription = *UseSourceMeshDescription;
+
+	// if not the high-res source, compute autogenerated normals/tangents
+	if (bUsingHiResSource == false)
+	{
+		UE::MeshDescription::InitializeAutoGeneratedAttributes(*SourceMeshDescription, SourceStaticMesh, 0);
+	}
+
+	// start async mesh-conversion
+	SourceMesh.Clear();
+	TFuture<void> ConvertMesh = Async(GenerateSMLODAsyncExecTarget, [&]()
+	{
+		FMeshDescriptionToDynamicMesh GetSourceMesh;
+		GetSourceMesh.Convert(SourceMeshDescription.Get(), SourceMesh);
+	});
+
+	// get list of source materials and find all the input texture params
+	const TArray<FStaticMaterial>& Materials = SourceStaticMesh->GetStaticMaterials();
+
+	// warn the user if there are any unsed materials in the mesh
+	if (Progress)
+	{
+		for (int32 UnusedMaterialIndex : GenerateStaticMeshLODProcessHelpers::FindUnreferencedMaterials(SourceStaticMesh, SourceMeshDescription.Get()))
+		{
+			const TObjectPtr<class UMaterialInterface> MaterialInterface = Materials[UnusedMaterialIndex].MaterialInterface;
+			if (ensure(MaterialInterface != nullptr))
+			{
+				FText WarningText = FText::Format(LOCTEXT("UnusedMaterialWarning", "Found an unused material ({0}). Consider removing it before using this tool."),
+					                              FText::FromName(MaterialInterface->GetFName()));
+				UE_LOG(LogMeshLODToolset, Warning, TEXT("%s"), *WarningText.ToString());
+				Progress->AddWarning(WarningText, FProgressCancel::EMessageLevel::UserWarning);
+			}
+		}
+	}
+
+
+	SourceMaterials.SetNum(Materials.Num());
+	TArray<FReadTextureJob> ReadJobs;
+	for (int32 mi = 0; mi < Materials.Num(); ++mi)
+	{
+		SourceMaterials[mi].SourceMaterial = Materials[mi];
+
+		UMaterialInterface* MaterialInterface = Materials[mi].MaterialInterface;
+		if (MaterialInterface == nullptr)
+		{
+			continue;
+		}
+
+		// detect hard-coded (non-parameter) texture samples
+		{
+			UMaterial* Material = Materials[mi].MaterialInterface->GetMaterial();
+
+			// go over the nodes in the material graph looking for texture samples
+			const UMaterialGraph* MatGraph = Material->MaterialGraph;
+
+			if (MatGraph == nullptr)	// create a material graph from the material if necessary
+			{
+				UMaterialGraph* NewMatGraph = CastChecked<UMaterialGraph>(NewObject<UEdGraph>(Material, UMaterialGraph::StaticClass(), NAME_None, RF_Transient));
+				NewMatGraph->Material = Material;
+				NewMatGraph->RebuildGraph();
+				MatGraph = NewMatGraph;
+			}
+
+			bool bFoundTextureNonParamExpession = false;
+			const TArray<TObjectPtr<class UEdGraphNode>>& Nodes = MatGraph->Nodes;
+			for (const TObjectPtr<UEdGraphNode>& Node : Nodes)
+			{
+				const UMaterialGraphNode* GraphNode = Cast<UMaterialGraphNode>(Node);
+				if (GraphNode)
+				{
+					const UMaterialExpressionTextureSampleParameter* TextureSampleParameterBase = Cast<UMaterialExpressionTextureSampleParameter>(GraphNode->MaterialExpression);
+					if (!TextureSampleParameterBase)
+					{
+						const UMaterialExpressionTextureSample* TextureSampleBase = Cast<UMaterialExpressionTextureSample>(GraphNode->MaterialExpression);
+						if (TextureSampleBase)
+						{
+							// node is UMaterialExpressionTextureSample but not UMaterialExpressionTextureSampleParameter
+							bFoundTextureNonParamExpession = true;
+							break;
+						}
+					}
+
+				}
+			}
+			if (bFoundTextureNonParamExpession)
+			{
+				FText WarningText = FText::Format(LOCTEXT("NonParameterTextureWarning", "Non-parameter texture sampler detected in input material [{0}]. Output materials may have unexpected behaviour."),
+												  FText::FromString(Material->GetName()));
+				UE_LOG(LogMeshLODToolset, Warning, TEXT("%s"), *WarningText.ToString());
+				if (Progress)
+				{
+					Progress->AddWarning(WarningText, FProgressCancel::EMessageLevel::UserWarning);
+				}
+			}
+		}
+
+		SourceMaterials[mi].bHasNormalMap = false;
+		SourceMaterials[mi].bHasTexturesToBake = false;
+
+		TArray<FMaterialParameterInfo> ParameterInfo;
+		TArray<FGuid> ParameterIds;
+		MaterialInterface->GetAllTextureParameterInfo(ParameterInfo, ParameterIds);
+		for (int32 ti = 0; ti < ParameterInfo.Num(); ++ti)
+		{
+			FName ParamName = ParameterInfo[ti].Name;
+
+			UTexture* CurTexture = nullptr;
+			bool bFound = MaterialInterface->GetTextureParameterValue(
+				FMemoryImageMaterialParameterInfo(ParameterInfo[ti]), CurTexture);
+			if (ensure(bFound))
+			{
+				if (Cast<UTexture2D>(CurTexture) != nullptr)
+				{
+					FTextureInfo TexInfo;
+					TexInfo.Texture = Cast<UTexture2D>(CurTexture);
+					TexInfo.ParameterName = ParamName;
+
+					TexInfo.bIsNormalMap = TexInfo.Texture->IsNormalMap();
+					SourceMaterials[mi].bHasNormalMap |= TexInfo.bIsNormalMap;
+
+					TexInfo.bIsDefaultTexture = UEditorAssetLibrary::GetPathNameForLoadedAsset(TexInfo.Texture).StartsWith(TEXT("/Engine/"));
+
+					TexInfo.bShouldBakeTexture = (TexInfo.bIsNormalMap == false && TexInfo.bIsDefaultTexture == false);
+					if (TexInfo.bShouldBakeTexture)
+					{
+						ReadJobs.Add(FReadTextureJob{ mi, SourceMaterials[mi].SourceTextures.Num() });
+
+						SourceMaterials[mi].bHasTexturesToBake = true;
+					}
+
+					SourceMaterials[mi].SourceTextures.Add(TexInfo);
+				}
+			}
+		}
+
+		// if material does not have a normal map parameter or any textures we want to bake, we can just re-use it
+		SourceMaterials[mi].bIsReusable = (SourceMaterials[mi].bHasNormalMap == false && SourceMaterials[mi].bHasTexturesToBake == false);
+	}
+
+
+	// If we have hi-res source we can discard any materials that are only used on the previously-generated LOD0.
+	// We cannot explicitly tag the materials as being generated so we infer, ie we assume a material was generated
+	// if it is only used in LOD0 and not HiRes
+	if (bUsingHiResSource)
+	{
+		// have to wait for SourceMesh to finish converting
+		ConvertMesh.Wait();
+
+		const FMeshDescription* LOD0MeshDescription = SourceStaticMesh->GetMeshDescription(0);
+		FMeshDescriptionToDynamicMesh GetLOD0Mesh;
+		FDynamicMesh3 LOD0Mesh;
+		GetLOD0Mesh.Convert(LOD0MeshDescription, LOD0Mesh);
+		const FDynamicMeshMaterialAttribute* SourceMaterialIDs = SourceMesh.Attributes()->GetMaterialID();
+		const FDynamicMeshMaterialAttribute* LOD0MaterialIDs = LOD0Mesh.Attributes()->GetMaterialID();
+		if (ensure(SourceMaterialIDs != nullptr && LOD0MaterialIDs != nullptr))
+		{
+			int32 NumMaterials = SourceMaterials.Num();
+			TArray<bool> IsBaseMaterial;
+			IsBaseMaterial.Init(false, NumMaterials);
+			TArray<bool> IsLOD0Material;
+			IsLOD0Material.Init(false, NumMaterials);
+
+			for (int32 tid : SourceMesh.TriangleIndicesItr())
+			{
+				int32 MatIdx = SourceMaterialIDs->GetValue(tid);
+				if (MatIdx >= 0 && MatIdx < NumMaterials)
+				{
+					IsBaseMaterial[MatIdx] = true;
+				}
+			}
+			for (int32 tid : LOD0Mesh.TriangleIndicesItr())
+			{
+				int32 MatIdx = LOD0MaterialIDs->GetValue(tid);
+				if (MatIdx >= 0 && MatIdx < NumMaterials)
+				{
+					IsLOD0Material[MatIdx] = true;
+				}
+			}
+
+			for (int32 k = 0; k < NumMaterials; ++k)
+			{
+				if (IsLOD0Material[k] == true && IsBaseMaterial[k] == false)
+				{
+					SourceMaterials[k].bIsPreviouslyGeneratedMaterial = true;
+					SourceMaterials[k].bHasTexturesToBake = false;
+					SourceMaterials[k].bIsReusable = false;
+				}
+			}
+		}
+	}
+
+
+
+	// extract all the texture params
+	// TODO: this triggers a checkSlow in the serialization code when it runs async. Find out why. Jira: UETOOL-2985
+	//TFuture<void> ReadTextures = Async(GenerateSMLODAsyncExecTarget, [&]()
+	//{
+	//	ParallelFor(ReadJobs.Num(), [&](int32 ji)
+	//	{
+	//		FReadTextureJob job = ReadJobs[ji];
+	//		FTextureInfo& TexInfo = SourceMaterials[job.MatIndex].SourceTextures[job.TexIndex];
+	//		UE::AssetUtils::ReadTexture(TexInfo.Texture, TexInfo.Image);
+	//		TexInfo.Dimensions = TexInfo.Image.GetDimensions();
+	//	});
+	//});
+
+	// single-thread path
+	TArray<UTexture2D*> UniqueTextures;
+	TArray<TImageBuilder<FVector4f>*> UniqueTextureImages;
+	for (FReadTextureJob job : ReadJobs)
+	{
+		// only read textures that are from materials we are going to possibly bake
+		FSourceMaterialInfo& SourceMaterial = SourceMaterials[job.MatIndex];
+		if (SourceMaterial.bIsPreviouslyGeneratedMaterial == false && SourceMaterial.bIsReusable == false)
+		{
+			FTextureInfo& TexInfo = SourceMaterials[job.MatIndex].SourceTextures[job.TexIndex];
+			UE::AssetUtils::ReadTexture(TexInfo.Texture, TexInfo.Image);
+			TexInfo.Dimensions = TexInfo.Image.GetDimensions();
+
+			UniqueTextures.AddUnique(TexInfo.Texture);
+			UniqueTextureImages.Add(&TexInfo.Image);
+		}
+	}
+
+	//ReadTextures.Wait();
+
+
+	// Now that we have read textures, detect solid-color textures. We do not
+	// have to bake these. These are often used for (eg) default textures, etc.
+	TArray<bool> bIsSolidColor;
+	bIsSolidColor.Init(false, UniqueTextures.Num());
+	ParallelFor(UniqueTextures.Num(), [&](int32 ti)
+	{
+		bIsSolidColor[ti] = UniqueTextureImages[ti]->IsConstantValue();
+	});
+	for (FSourceMaterialInfo& SourceMaterial : SourceMaterials)
+	{
+		if (SourceMaterial.bHasTexturesToBake)
+		{
+			int StillToBakeCount = 0;
+			for (FTextureInfo& TexInfo : SourceMaterial.SourceTextures)
+			{
+				int32 Index = UniqueTextures.IndexOfByKey(TexInfo.Texture);
+				if (Index != INDEX_NONE && bIsSolidColor[Index])
+				{
+					TexInfo.bShouldBakeTexture = false;
+				}
+				else
+				{
+					StillToBakeCount++;
+				}
+			}
+			if (StillToBakeCount == 0)
+			{
+				// If Material don't have any remaining textures to bake, and it doesn't have a normal map, then we do not have to bake it
+				SourceMaterial.bHasTexturesToBake = false;
+				SourceMaterial.bIsReusable = (SourceMaterial.bHasNormalMap == false);
+			}
+		}
+	}
+
+
+	ConvertMesh.Wait();
+
+	FString FullPathWithExtension = UEditorAssetLibrary::GetPathNameForLoadedAsset(SourceStaticMesh);
+	SourceAssetPath = FPaths::GetBaseFilename(FullPathWithExtension, false);
+	SourceAssetFolder = FPaths::GetPath(SourceAssetPath);
+	SourceAssetName = FPaths::GetBaseFilename(FullPathWithExtension, true);
+
+	UpdateDerivedPathName(SourceAssetName, GetDefaultDerivedAssetSuffix());
+
+	InitializeGenerator();
+
+	return true;
+}
+
+
+void UGenerateStaticMeshLODProcess::UpdateDerivedPathName(const FString& NewAssetBaseName, const FString& NewAssetSuffix)
+{
+	DerivedAssetNameNoSuffix = FPaths::MakeValidFileName(NewAssetBaseName.Replace(TEXT(" "), TEXT("_")));
+	if (DerivedAssetNameNoSuffix.Len() == 0)
+	{
+		DerivedAssetNameNoSuffix = SourceAssetName;
+	}
+
+	DerivedSuffix = FPaths::MakeValidFileName(NewAssetSuffix.Replace(TEXT(" "), TEXT("_")));
+	if (DerivedSuffix.Len() == 0)
+	{
+		DerivedSuffix = GetDefaultDerivedAssetSuffix();
+	}
+
+	DerivedAssetName = FString::Printf(TEXT("%s%s"), *DerivedAssetNameNoSuffix, *DerivedSuffix);
+	DerivedAssetFolder = SourceAssetFolder;
+	DerivedAssetPath = FPaths::Combine(DerivedAssetFolder, DerivedAssetName);
+}
+
+
+int UGenerateStaticMeshLODProcess::SelectTextureToBake(const TArray<FTextureInfo>& TextureInfos) const
+{
+	TArray<int> TextureVotes;
+	TextureVotes.Init(0, TextureInfos.Num());
+
+	for (int TextureIndex = 0; TextureIndex < TextureInfos.Num(); ++TextureIndex)
+	{
+		FTextureInfo Info = TextureInfos[TextureIndex];
+
+		if (!Info.bIsNormalMap)
+		{
+			++TextureVotes[TextureIndex];
+		}
+
+		if (Info.ParameterName == "Diffuse" || Info.ParameterName == "Albedo" || Info.ParameterName == "BaseColor")
+		{
+			++TextureVotes[TextureIndex];
+		}
+
+		if (Info.Dimensions.GetWidth() > 0 && Info.Dimensions.GetHeight() > 0)
+		{
+			++TextureVotes[TextureIndex];
+		}
+
+		UTexture2D* Tex2D = Info.Texture;
+
+		if (Tex2D)
+		{
+			// Texture uses SRGB
+			if (Tex2D->SRGB != 0)
+			{
+				++TextureVotes[TextureIndex];
+			}
+
+#if WITH_EDITORONLY_DATA
+			// Texture has multiple channels
+			ETextureSourceFormat Format = Tex2D->Source.GetFormat();
+			if (Format == TSF_BGRA8 || Format == TSF_BGRE8 || Format == TSF_RGBA16 || Format == TSF_RGBA16F || Format == TSF_RGBA32F)
+			{
+				++TextureVotes[TextureIndex];
+			}
+#endif
+
+			// What else? Largest texture? Most layers? Most mipmaps?
+		}
+	}
+
+	int MaxIndex = -1;
+	int MaxVotes = -1;
+	for (int TextureIndex = 0; TextureIndex < TextureInfos.Num(); ++TextureIndex)
+	{
+		if (TextureVotes[TextureIndex] > MaxVotes)
+		{
+			MaxIndex = TextureIndex;
+			MaxVotes = TextureVotes[TextureIndex];
+		}
+	}
+
+	return MaxIndex;
+}
+
+
+
+bool UGenerateStaticMeshLODProcess::InitializeGenerator()
+{
+	Generator = MakeUnique<FGenerateMeshLODGraph>();
+	Generator->BuildGraph( & this->SourceMesh );
+
+	// initialize source textures that need to be baked
+	SourceTextureToDerivedTexIndex.Reset();
+	for (const FSourceMaterialInfo& MatInfo : SourceMaterials)
+	{
+		if ( MatInfo.bIsPreviouslyGeneratedMaterial == false && MatInfo.bIsReusable == false && MatInfo.bHasTexturesToBake)
+		{
+			for (const FTextureInfo& TexInfo : MatInfo.SourceTextures)
+			{
+				if (TexInfo.bShouldBakeTexture &&
+					(SourceTextureToDerivedTexIndex.Contains(TexInfo.Texture) == false))
+				{
+					int32 NewIndex = Generator->AppendTextureBakeNode(TexInfo.Image, TexInfo.Texture->GetName());
+					SourceTextureToDerivedTexIndex.Add(TexInfo.Texture, NewIndex);
+				}
+			}
+		}
+	}
+
+	//  Add multi-texture bake node	
+	TMap<int32, TSafeSharedPtr<UE::Geometry::TImageBuilder<FVector4f>>> SourceMaterialImages;
+
+	// TODO: Relying on MaterialID being the index in the SourceMaterials array is brittle
+	for (int32 MaterialID = 0; MaterialID < SourceMaterials.Num(); ++MaterialID)
+	{
+		FSourceMaterialInfo& MatInfo = SourceMaterials[MaterialID];
+
+		if (MatInfo.bIsPreviouslyGeneratedMaterial == false && MatInfo.bIsReusable == false && MatInfo.bHasTexturesToBake)
+		{
+			int TexIndex = SelectTextureToBake(MatInfo.SourceTextures);
+			if (TexIndex < 0)
+			{
+				continue;
+			}
+
+			FTextureInfo& TextureInfo = MatInfo.SourceTextures[TexIndex];
+
+			MultiTextureParameterName.Add(MaterialID, TextureInfo.ParameterName);
+
+			TSafeSharedPtr<TImageBuilder<FVector4f>> ImagePtr =
+				MakeSafeShared<TImageBuilder<FVector4f>>(TextureInfo.Image);
+			SourceMaterialImages.Add(MaterialID, ImagePtr);
+
+			TextureInfo.bIsUsedInMultiTextureBaking = true;
+		}
+	}
+
+	Generator->AppendMultiTextureBakeNode(SourceMaterialImages);
+
+
+	// initialize source mesh
+	Generator->SetSourceMesh(this->SourceMesh);
+
+
+	// read parameter settings from the graph
+
+	CurrentSettings_Preprocess.FilterGroupLayer = Generator->GetCurrentPreFilterSettings().FilterGroupLayerName;
+	CurrentSettings_Preprocess.ThickenAmount = Generator->GetCurrentThickenSettings().ThickenAmount;
+	CurrentSettings.MeshGenerator = static_cast<EGenerateStaticMeshLODProcess_MeshGeneratorModes>(Generator->GetCurrentCoreMeshGeneratorMode());
+
+	CurrentSettings.SolidifyVoxelResolution = Generator->GetCurrentSolidifySettings().VoxelResolution;
+	CurrentSettings.WindingThreshold = Generator->GetCurrentSolidifySettings().WindingThreshold;
+	CurrentSettings.ClosureDistance = Generator->GetCurrentMorphologySettings().Distance;
+	CurrentSettings.bPrefilterVertices = Generator->GetCurrentGenerateConvexHullMeshSettings().bPrefilterVertices;
+	CurrentSettings.PrefilterGridResolution = Generator->GetCurrentGenerateConvexHullMeshSettings().PrefilterGridResolution;
+
+	CurrentSettings_Simplify.Method = static_cast<EGenerateStaticMeshLODProcess_SimplifyMethod>(Generator->GetCurrentSimplifySettings().TargetType);
+	CurrentSettings_Simplify.TargetCount = Generator->GetCurrentSimplifySettings().TargetCount;
+	CurrentSettings_Simplify.TargetPercentage = Generator->GetCurrentSimplifySettings().TargetFraction * 100.0f;
+	CurrentSettings_Simplify.Tolerance = Generator->GetCurrentSimplifySettings().GeometricTolerance;
+
+	CurrentSettings_Normals.Method = static_cast<EGenerateStaticMeshLODProcess_NormalsMethod>(
+		FGenerateStaticMeshLODProcess_NormalsSettings::MapMethodType((int32)Generator->GetCurrentNormalsSettings().NormalsType, false));
+	CurrentSettings_Normals.Angle = Generator->GetCurrentNormalsSettings().AngleThresholdDeg;
+
+	CurrentSettings_Texture.BakeResolution = (EGenerateStaticMeshLODBakeResolution)Generator->GetCurrentBakeCacheSettings().Dimensions.GetWidth();
+	CurrentSettings_Texture.BakeThickness = Generator->GetCurrentBakeCacheSettings().Thickness;
+
+	CurrentSettings_UV.UVMethod = static_cast<EGenerateStaticMeshLODProcess_AutoUVMethod>(Generator->GetCurrentAutoUVSettings().Method);
+	CurrentSettings_UV.NumUVAtlasCharts = Generator->GetCurrentAutoUVSettings().UVAtlasNumCharts;
+	CurrentSettings_UV.NumInitialPatches = Generator->GetCurrentAutoUVSettings().NumInitialPatches;
+	CurrentSettings_UV.MergingThreshold = Generator->GetCurrentAutoUVSettings().MergingThreshold;
+	CurrentSettings_UV.MaxAngleDeviation = Generator->GetCurrentAutoUVSettings().MaxAngleDeviationDeg;
+	CurrentSettings_UV.PatchBuilder.CurvatureAlignment = Generator->GetCurrentAutoUVSettings().CurvatureAlignment;
+	CurrentSettings_UV.PatchBuilder.SmoothingSteps = Generator->GetCurrentAutoUVSettings().SmoothingSteps;
+	CurrentSettings_UV.PatchBuilder.SmoothingAlpha = Generator->GetCurrentAutoUVSettings().SmoothingAlpha;
+
+	const UE::GeometryFlow::FGenerateSimpleCollisionSettings& SimpleCollisionSettings = Generator->GetCurrentGenerateSimpleCollisionSettings();
+	CurrentSettings_Collision.CollisionType = static_cast<EGenerateStaticMeshLODSimpleCollisionGeometryType>(SimpleCollisionSettings.Type);
+	CurrentSettings_Collision.ConvexTriangleCount = SimpleCollisionSettings.ConvexHullSettings.SimplifyToTriangleCount;
+	CurrentSettings_Collision.bPrefilterVertices = SimpleCollisionSettings.ConvexHullSettings.bPrefilterVertices;
+	CurrentSettings_Collision.PrefilterGridResolution = SimpleCollisionSettings.ConvexHullSettings.PrefilterGridResolution;
+	CurrentSettings_Collision.bSimplifyPolygons = SimpleCollisionSettings.SweptHullSettings.bSimplifyPolygons;
+	CurrentSettings_Collision.HullTolerance = SimpleCollisionSettings.SweptHullSettings.HullTolerance;
+	FMeshSimpleShapeApproximation::EProjectedHullAxisMode RHSMode = SimpleCollisionSettings.SweptHullSettings.SweepAxis;
+	CurrentSettings_Collision.SweepAxis = static_cast<EGenerateStaticMeshLODProjectedHullAxisMode>(RHSMode);
+
+
+	return true;
+}
+
+
+void UGenerateStaticMeshLODProcess::UpdateSettings(const FGenerateStaticMeshLODProcessSettings& NewCombinedSettings)
+{
+	if (NewCombinedSettings.MeshGenerator != CurrentSettings.MeshGenerator)
+	{
+		Generator->UpdateCoreMeshGeneratorMode(static_cast<FGenerateMeshLODGraph::ECoreMeshGeneratorMode>(NewCombinedSettings.MeshGenerator));
+	}
+
+	bool bSharedVoxelResolutionChanged = (NewCombinedSettings.SolidifyVoxelResolution != CurrentSettings.SolidifyVoxelResolution);
+	if ( bSharedVoxelResolutionChanged
+		|| (NewCombinedSettings.WindingThreshold != CurrentSettings.WindingThreshold))
+	{
+		UE::GeometryFlow::FMeshSolidifySettings NewSolidifySettings = Generator->GetCurrentSolidifySettings();
+		NewSolidifySettings.VoxelResolution = NewCombinedSettings.SolidifyVoxelResolution;
+		NewSolidifySettings.WindingThreshold = NewCombinedSettings.WindingThreshold;
+		Generator->UpdateSolidifySettings(NewSolidifySettings);
+	}
+
+
+	if ( bSharedVoxelResolutionChanged
+		|| (NewCombinedSettings.ClosureDistance != CurrentSettings.ClosureDistance))
+	{
+		UE::GeometryFlow::FVoxClosureSettings NewClosureSettings = Generator->GetCurrentMorphologySettings();
+		NewClosureSettings.VoxelResolution = NewCombinedSettings.SolidifyVoxelResolution;
+		NewClosureSettings.Distance = NewCombinedSettings.ClosureDistance;
+		Generator->UpdateMorphologySettings(NewClosureSettings);
+	}
+
+	if (NewCombinedSettings.bPrefilterVertices != CurrentSettings.bPrefilterVertices ||
+		NewCombinedSettings.PrefilterGridResolution != CurrentSettings.PrefilterGridResolution)
+	{
+		UE::GeometryFlow::FGenerateConvexHullMeshSettings NewConvexHullSettings = Generator->GetCurrentGenerateConvexHullMeshSettings();
+		NewConvexHullSettings.bPrefilterVertices = NewCombinedSettings.bPrefilterVertices;
+		NewConvexHullSettings.PrefilterGridResolution = NewCombinedSettings.PrefilterGridResolution;
+		Generator->UpdateGenerateConvexHullMeshSettings(NewConvexHullSettings);
+	}
+
+	CurrentSettings = NewCombinedSettings;
+}
+
+
+
+void UGenerateStaticMeshLODProcess::UpdatePreprocessSettings(const FGenerateStaticMeshLODProcess_PreprocessSettings& NewCombinedSettings)
+{
+	if (NewCombinedSettings.FilterGroupLayer != CurrentSettings_Preprocess.FilterGroupLayer)
+	{
+		FMeshLODGraphPreFilterSettings NewPreFilterSettings = Generator->GetCurrentPreFilterSettings();
+		NewPreFilterSettings.FilterGroupLayerName = NewCombinedSettings.FilterGroupLayer;
+		Generator->UpdatePreFilterSettings(NewPreFilterSettings);
+	}
+
+	if (NewCombinedSettings.ThickenAmount != CurrentSettings_Preprocess.ThickenAmount)
+	{
+		UE::GeometryFlow::FMeshThickenSettings NewThickenSettings = Generator->GetCurrentThickenSettings();
+		NewThickenSettings.ThickenAmount = NewCombinedSettings.ThickenAmount;
+		Generator->UpdateThickenSettings(NewThickenSettings);
+	}
+
+	if (NewCombinedSettings.ThickenWeightMapName != CurrentSettings_Preprocess.ThickenWeightMapName)
+	{
+		FIndexedWeightMap1f WeightMap;
+		float DefaultValue = 0.0f;
+		bool bOK = UE::WeightMaps::GetVertexWeightMap(SourceMeshDescription.Get(), NewCombinedSettings.ThickenWeightMapName, WeightMap, DefaultValue);
+
+		if (bOK)
+		{
+			Generator->UpdateThickenWeightMap(WeightMap.Values);
+		}
+		else
+		{
+			Generator->UpdateThickenWeightMap(TArray<float>());
+		}
+	}
+
+	CurrentSettings_Preprocess = NewCombinedSettings;
+}
+
+
+
+void UGenerateStaticMeshLODProcess::UpdateSimplifySettings(const FGenerateStaticMeshLODProcess_SimplifySettings& NewCombinedSettings)
+{
+	if (NewCombinedSettings.Method != CurrentSettings_Simplify.Method || 
+		NewCombinedSettings.TargetCount != CurrentSettings_Simplify.TargetCount ||
+		NewCombinedSettings.TargetPercentage != CurrentSettings_Simplify.TargetPercentage ||
+		NewCombinedSettings.Tolerance != CurrentSettings_Simplify.Tolerance )
+	{
+		UE::GeometryFlow::FMeshSimplifySettings NewSimplifySettings = Generator->GetCurrentSimplifySettings();
+		NewSimplifySettings.TargetType = static_cast<UE::GeometryFlow::EMeshSimplifyTargetType>(NewCombinedSettings.Method);
+		NewSimplifySettings.TargetCount = NewCombinedSettings.TargetCount;
+		NewSimplifySettings.TargetFraction = NewCombinedSettings.TargetPercentage / 100.0f;
+		NewSimplifySettings.GeometricTolerance = NewCombinedSettings.Tolerance;
+		Generator->UpdateSimplifySettings(NewSimplifySettings);
+	}
+
+	CurrentSettings_Simplify = NewCombinedSettings;
+}
+
+
+
+int32 FGenerateStaticMeshLODProcess_NormalsSettings::MapMethodType(int32 From, bool bProcessToGraph)
+{
+	if (bProcessToGraph)
+	{
+		switch (static_cast<EGenerateStaticMeshLODProcess_NormalsMethod>(From))
+		{
+			default:
+			case EGenerateStaticMeshLODProcess_NormalsMethod::FromAngleThreshold: return (int32)UE::GeometryFlow::EComputeNormalsType::FromFaceAngleThreshold;
+			case EGenerateStaticMeshLODProcess_NormalsMethod::PerVertex: return (int32)UE::GeometryFlow::EComputeNormalsType::PerVertex;
+			case EGenerateStaticMeshLODProcess_NormalsMethod::PerTriangle: return (int32)UE::GeometryFlow::EComputeNormalsType::PerTriangle;
+		}
+	}
+	else
+	{
+		switch (static_cast<UE::GeometryFlow::EComputeNormalsType>(From))
+		{
+			default:
+			case UE::GeometryFlow::EComputeNormalsType::FromFaceAngleThreshold: return (int32)EGenerateStaticMeshLODProcess_NormalsMethod::FromAngleThreshold;
+			case UE::GeometryFlow::EComputeNormalsType::PerVertex: return (int32)EGenerateStaticMeshLODProcess_NormalsMethod::PerVertex;
+			case UE::GeometryFlow::EComputeNormalsType::PerTriangle: return (int32)EGenerateStaticMeshLODProcess_NormalsMethod::PerTriangle;
+		}
+	}
+}
+
+
+void UGenerateStaticMeshLODProcess::UpdateNormalsSettings(const FGenerateStaticMeshLODProcess_NormalsSettings& NewCombinedSettings)
+{
+	if (NewCombinedSettings.Method != CurrentSettings_Normals.Method || 
+		NewCombinedSettings.Angle != CurrentSettings_Normals.Angle )
+	{
+		UE::GeometryFlow::FMeshNormalsSettings NewNormalsSettings = Generator->GetCurrentNormalsSettings();
+
+		NewNormalsSettings.NormalsType = static_cast<UE::GeometryFlow::EComputeNormalsType>(
+			FGenerateStaticMeshLODProcess_NormalsSettings::MapMethodType((int32)NewCombinedSettings.Method, true));
+		NewNormalsSettings.AngleThresholdDeg = NewCombinedSettings.Angle;
+
+		Generator->UpdateNormalsSettings(NewNormalsSettings);
+	}
+
+	CurrentSettings_Normals = NewCombinedSettings;
+}
+
+
+
+void UGenerateStaticMeshLODProcess::UpdateTextureSettings(const FGenerateStaticMeshLODProcess_TextureSettings& NewCombinedSettings)
+{
+	if ( (NewCombinedSettings.BakeResolution != CurrentSettings_Texture.BakeResolution) ||
+		 (NewCombinedSettings.BakeThickness != CurrentSettings_Texture.BakeThickness) || 
+		 (NewCombinedSettings.bCombineTextures != CurrentSettings_Texture.bCombineTextures))
+	{
+		UE::GeometryFlow::FMeshMakeBakingCacheSettings NewBakeSettings = Generator->GetCurrentBakeCacheSettings();
+		NewBakeSettings.Dimensions = FImageDimensions((int32)NewCombinedSettings.BakeResolution, (int32)NewCombinedSettings.BakeResolution);
+		NewBakeSettings.Thickness = NewCombinedSettings.BakeThickness;
+		Generator->UpdateBakeCacheSettings(NewBakeSettings);
+	}
+
+	CurrentSettings_Texture = NewCombinedSettings;
+}
+
+
+
+
+void UGenerateStaticMeshLODProcess::UpdateUVSettings(const FGenerateStaticMeshLODProcess_UVSettings& NewCombinedSettings)
+{
+	if (NewCombinedSettings.UVMethod != CurrentSettings_UV.UVMethod
+		|| NewCombinedSettings.NumUVAtlasCharts != CurrentSettings_UV.NumUVAtlasCharts
+		|| NewCombinedSettings.NumInitialPatches != CurrentSettings_UV.NumInitialPatches
+		|| NewCombinedSettings.MergingThreshold != CurrentSettings_UV.MergingThreshold
+		|| NewCombinedSettings.MaxAngleDeviation != CurrentSettings_UV.MaxAngleDeviation
+		|| NewCombinedSettings.PatchBuilder.CurvatureAlignment != CurrentSettings_UV.PatchBuilder.CurvatureAlignment
+		|| NewCombinedSettings.PatchBuilder.SmoothingSteps != CurrentSettings_UV.PatchBuilder.SmoothingSteps
+		|| NewCombinedSettings.PatchBuilder.SmoothingAlpha != CurrentSettings_UV.PatchBuilder.SmoothingAlpha  ) 
+	{
+		UE::GeometryFlow::FMeshAutoGenerateUVsSettings NewAutoUVSettings = Generator->GetCurrentAutoUVSettings();
+		NewAutoUVSettings.Method = (EAutoUVMethod)(int)NewCombinedSettings.UVMethod;
+		//NewAutoUVSettings.UVAtlasStretch = 0.5;
+		NewAutoUVSettings.UVAtlasNumCharts = NewCombinedSettings.NumUVAtlasCharts;
+		//NewAutoUVSettings.XAtlasMaxIterations = 1;
+		NewAutoUVSettings.NumInitialPatches = NewCombinedSettings.NumInitialPatches;
+		NewAutoUVSettings.CurvatureAlignment = NewCombinedSettings.PatchBuilder.CurvatureAlignment;
+		NewAutoUVSettings.MergingThreshold = NewCombinedSettings.MergingThreshold;
+		NewAutoUVSettings.MaxAngleDeviationDeg = NewCombinedSettings.MaxAngleDeviation;
+		NewAutoUVSettings.SmoothingSteps = NewCombinedSettings.PatchBuilder.SmoothingSteps;
+		NewAutoUVSettings.SmoothingAlpha = NewCombinedSettings.PatchBuilder.SmoothingAlpha;
+		//NewAutoUVSettings.bAutoPack = false;
+		//NewAutoUVSettings.PackingTargetWidth = 512;
+		Generator->UpdateAutoUVSettings(NewAutoUVSettings);
+	}
+
+	CurrentSettings_UV = NewCombinedSettings;
+}
+
+
+
+
+void UGenerateStaticMeshLODProcess::UpdateCollisionSettings(const FGenerateStaticMeshLODProcess_CollisionSettings& NewCombinedSettings)
+{
+	if (NewCombinedSettings.CollisionGroupLayerName != CurrentSettings_Collision.CollisionGroupLayerName)
+	{
+		Generator->UpdateCollisionGroupLayerName(NewCombinedSettings.CollisionGroupLayerName);
+	}
+
+	if (NewCombinedSettings.ConvexTriangleCount != CurrentSettings_Collision.ConvexTriangleCount ||
+		NewCombinedSettings.bPrefilterVertices != CurrentSettings_Collision.bPrefilterVertices ||
+		NewCombinedSettings.PrefilterGridResolution != CurrentSettings_Collision.PrefilterGridResolution ||
+		NewCombinedSettings.bSimplifyPolygons != CurrentSettings_Collision.bSimplifyPolygons ||
+		NewCombinedSettings.HullTolerance != CurrentSettings_Collision.HullTolerance ||
+		NewCombinedSettings.SweepAxis != CurrentSettings_Collision.SweepAxis ||
+		NewCombinedSettings.CollisionType != CurrentSettings_Collision.CollisionType)
+	{
+		UE::GeometryFlow::FGenerateSimpleCollisionSettings NewGenCollisionSettings = Generator->GetCurrentGenerateSimpleCollisionSettings();
+		NewGenCollisionSettings.Type = static_cast<UE::GeometryFlow::ESimpleCollisionGeometryType>(NewCombinedSettings.CollisionType);
+		NewGenCollisionSettings.ConvexHullSettings.SimplifyToTriangleCount = NewCombinedSettings.ConvexTriangleCount;
+		NewGenCollisionSettings.ConvexHullSettings.bPrefilterVertices = NewCombinedSettings.bPrefilterVertices;
+		NewGenCollisionSettings.ConvexHullSettings.PrefilterGridResolution = NewCombinedSettings.PrefilterGridResolution;
+		NewGenCollisionSettings.SweptHullSettings.bSimplifyPolygons = NewCombinedSettings.bSimplifyPolygons;
+		NewGenCollisionSettings.SweptHullSettings.HullTolerance = NewCombinedSettings.HullTolerance;
+		NewGenCollisionSettings.SweptHullSettings.SweepAxis = static_cast<FMeshSimpleShapeApproximation::EProjectedHullAxisMode>(NewCombinedSettings.SweepAxis);
+		Generator->UpdateGenerateSimpleCollisionSettings(NewGenCollisionSettings);
+	}
+
+	CurrentSettings_Collision = NewCombinedSettings;
+}
+
+
+
+TArray<UMaterialInterface*> UGenerateStaticMeshLODProcess::GetSourceBakeMaterials() const
+{
+	TArray<UMaterialInterface*> UniqueMaterials;
+	for (const FSourceMaterialInfo& MatInfo : SourceMaterials)
+	{
+		if (MatInfo.bHasTexturesToBake)
+		{
+			UniqueMaterials.AddUnique(MatInfo.SourceMaterial.MaterialInterface);
+		}
+	}
+	return UniqueMaterials;
+}
+
+void UGenerateStaticMeshLODProcess::UpdateSourceBakeMaterialConstraint(UMaterialInterface* Material, EMaterialBakingConstraint Constraint)
+{
+	for (FSourceMaterialInfo& MatInfo : SourceMaterials)
+	{
+		if (MatInfo.SourceMaterial.MaterialInterface == Material)
+		{
+			MatInfo.Constraint = Constraint;
+		}
+	}
+}
+
+
+TArray<UTexture2D*> UGenerateStaticMeshLODProcess::GetSourceBakeTextures() const
+{
+	TArray<UTexture2D*> UniqueTextures;
+	for (const FSourceMaterialInfo& MatInfo : SourceMaterials)
+	{
+		if (MatInfo.bHasTexturesToBake)
+		{
+			for (const FTextureInfo& TexInfo : MatInfo.SourceTextures)
+			{
+				if (TexInfo.bShouldBakeTexture)
+				{
+					UniqueTextures.AddUnique(TexInfo.Texture);
+				}
+			}
+		}
+	}
+	return UniqueTextures;
+}
+
+
+void UGenerateStaticMeshLODProcess::UpdateSourceBakeTextureConstraint(UTexture2D* Texture, ETextureBakingConstraint Constraint)
+{
+	for (FSourceMaterialInfo& MatInfo : SourceMaterials)
+	{
+		if (MatInfo.bHasTexturesToBake)
+		{
+			for (FTextureInfo& TexInfo : MatInfo.SourceTextures)
+			{
+				if (TexInfo.Texture == Texture)
+				{
+					TexInfo.Constraint = Constraint;
+				}
+			}
+		}
+	}
+}
+
+
+
+
+
+bool UGenerateStaticMeshLODProcess::ComputeDerivedSourceData(FProgressCancel* Progress)
+{
+	if (Generator == nullptr)
+	{
+		return false;
+	}
+
+	DerivedTextureImages.Reset();
+
+	Generator->EvaluateResult(
+		this->DerivedLODMesh,
+		this->DerivedLODMeshTangents,
+		this->DerivedCollision,
+		this->DerivedNormalMapImage,
+		this->DerivedTextureImages,
+		this->DerivedMultiTextureBakeImage,
+		Progress);
+
+	if (Progress && Progress->Cancelled())
+	{
+		return false;
+	}
+
+	// copy materials
+	int32 NumMaterials = SourceMaterials.Num();
+	DerivedMaterials.SetNum(NumMaterials);
+	for (int32 mi = 0; mi < NumMaterials; ++mi)
+	{
+		DerivedMaterials[mi].SourceMaterialIndex = mi;
+		DerivedMaterials[mi].bUseSourceMaterialDirectly = 
+			SourceMaterials[mi].bIsReusable 
+			|| SourceMaterials[mi].bIsPreviouslyGeneratedMaterial
+			|| (SourceMaterials[mi].Constraint == EMaterialBakingConstraint::UseExistingMaterial);
+		
+		if (DerivedMaterials[mi].bUseSourceMaterialDirectly)
+		{
+			DerivedMaterials[mi].DerivedMaterial = SourceMaterials[mi].SourceMaterial;
+		}
+		else
+		{
+			// TODO this is a lot of wasted overhead, we do not need to copy images here for example
+			DerivedMaterials[mi].DerivedTextures = SourceMaterials[mi].SourceTextures;
+		}
+	}
+
+	// update texture image data in derived materials
+	for (FDerivedMaterialInfo& MatInfo : DerivedMaterials)
+	{
+		for (FTextureInfo& TexInfo : MatInfo.DerivedTextures)
+		{
+			if (SourceTextureToDerivedTexIndex.Contains(TexInfo.Texture))
+			{
+				int32 BakedTexIndex = SourceTextureToDerivedTexIndex[TexInfo.Texture];
+				const TUniquePtr<UE::GeometryFlow::FTextureImage>& DerivedTex = DerivedTextureImages[BakedTexIndex];
+				TexInfo.Dimensions = DerivedTex->Image.GetDimensions();
+
+				// Cannot currently MoveTemp here because this Texture may appear in multiple Materials, 
+				// and currently we do not handle that. The Materials need to learn how to share.
+				//TexInfo.Image = MoveTemp(DerivedTex->Image);
+				TexInfo.Image = DerivedTex->Image;
+			}
+		}
+	}
+
+	return true;
+}
+
+
+
+
+void UGenerateStaticMeshLODProcess::GetDerivedMaterialsPreview(FPreviewMaterials& MaterialSetOut)
+{
+	// force garbage collection of outstanding preview materials
+	CollectGarbage(GARBAGE_COLLECTION_KEEPFLAGS);
+
+	bool bAnyMaterialUsedGeneratedNormalMap = false;
+
+	// create derived textures
+	int32 NumMaterials = SourceMaterials.Num();
+	check(DerivedMaterials.Num() == NumMaterials);
+	TMap<UTexture2D*, UTexture2D*> SourceToPreviewTexMap;
+	for (int32 mi = 0; mi < NumMaterials; ++mi)
+	{
+		const FSourceMaterialInfo& SourceMaterialInfo = SourceMaterials[mi];
+		if (SourceMaterialInfo.bIsReusable 
+			|| SourceMaterialInfo.bIsPreviouslyGeneratedMaterial
+			|| SourceMaterialInfo.Constraint == EMaterialBakingConstraint::UseExistingMaterial )
+		{
+			continue;
+		}
+
+		if (SourceMaterialInfo.bHasNormalMap)
+		{
+			bAnyMaterialUsedGeneratedNormalMap = true;
+		}
+
+		const FDerivedMaterialInfo& DerivedMaterialInfo = DerivedMaterials[mi];
+
+		int32 NumTextures = SourceMaterialInfo.SourceTextures.Num();
+		check(DerivedMaterialInfo.DerivedTextures.Num() == NumTextures);
+		for (int32 ti = 0; ti < NumTextures; ++ti)
+		{
+			const FTextureInfo& SourceTex = SourceMaterialInfo.SourceTextures[ti];
+
+			if (CurrentSettings_Texture.bCombineTextures && SourceTex.bIsUsedInMultiTextureBaking)
+			{
+				continue;
+			}
+
+			bool bConvertToSRGB = SourceTex.Texture->SRGB;
+			const FTextureInfo& DerivedTex = DerivedMaterialInfo.DerivedTextures[ti];
+			if (DerivedTex.bShouldBakeTexture
+				&& SourceTex.Constraint != ETextureBakingConstraint::UseExistingTexture )
+			{
+				FTexture2DBuilder TextureBuilder;
+				TextureBuilder.Initialize(FTexture2DBuilder::ETextureType::Color, DerivedTex.Dimensions);
+				TextureBuilder.GetTexture2D()->SRGB = bConvertToSRGB;
+				TextureBuilder.Copy(DerivedTex.Image, bConvertToSRGB);
+				TextureBuilder.Commit(false);
+				UTexture2D* PreviewTex = TextureBuilder.GetTexture2D();
+				if (ensure(PreviewTex))
+				{
+					SourceToPreviewTexMap.Add(SourceTex.Texture, PreviewTex);
+					MaterialSetOut.Textures.Add(PreviewTex);
+				}
+			}
+		}
+	}
+
+	// create derived normal map texture
+	UTexture2D* PreviewNormalMapTex = nullptr;
+	if (bAnyMaterialUsedGeneratedNormalMap)
+	{
+		FTexture2DBuilder NormapMapBuilder;
+		NormapMapBuilder.Initialize(FTexture2DBuilder::ETextureType::NormalMap, DerivedNormalMapImage.Image.GetDimensions());
+		NormapMapBuilder.Copy(DerivedNormalMapImage.Image, false);
+		NormapMapBuilder.Commit(false);
+		PreviewNormalMapTex = NormapMapBuilder.GetTexture2D();
+		MaterialSetOut.Textures.Add(PreviewNormalMapTex);
+	}
+
+	// create multi-texture bake result
+	if (CurrentSettings_Texture.bCombineTextures)
+	{
+		bool bConvertToSRGB = true;
+
+		FTexture2DBuilder MultiTextureBuilder;
+		MultiTextureBuilder.Initialize(FTexture2DBuilder::ETextureType::Color, DerivedMultiTextureBakeImage.Image.GetDimensions());
+		MultiTextureBuilder.GetTexture2D()->SRGB = bConvertToSRGB;
+		MultiTextureBuilder.Copy(DerivedMultiTextureBakeImage.Image, bConvertToSRGB);
+		MultiTextureBuilder.Commit(false);
+
+		DerivedMultiTextureBakeResult = MultiTextureBuilder.GetTexture2D();
+		MaterialSetOut.Textures.Add(DerivedMultiTextureBakeResult);
+	}
+
+
+	// create derived MIDs and point to new textures
+	for (int32 mi = 0; mi < NumMaterials; ++mi)
+	{
+		const FSourceMaterialInfo& SourceMaterialInfo = SourceMaterials[mi];
+		UMaterialInterface* MaterialInterface = SourceMaterialInfo.SourceMaterial.MaterialInterface;
+
+		if (SourceMaterialInfo.bIsReusable 
+			|| SourceMaterialInfo.bIsPreviouslyGeneratedMaterial
+			|| SourceMaterialInfo.Constraint == EMaterialBakingConstraint::UseExistingMaterial )
+		{
+			MaterialSetOut.Materials.Add(MaterialInterface);
+		}
+		else
+		{
+			// TODO: we should cache these instead of re-creating every time??
+			UMaterialInstanceDynamic* GeneratedMID = UMaterialInstanceDynamic::Create(MaterialInterface, NULL);
+
+			// rewrite texture parameters to new textures
+			UpdateMaterialTextureParameters(GeneratedMID, SourceMaterialInfo, SourceToPreviewTexMap, PreviewNormalMapTex);
+
+			if (CurrentSettings_Texture.bCombineTextures && DerivedMultiTextureBakeResult && MultiTextureParameterName.Contains(mi))
+			{
+				FMaterialParameterInfo ParamInfo(MultiTextureParameterName[mi]);
+				GeneratedMID->SetTextureParameterValueByInfo(ParamInfo, DerivedMultiTextureBakeResult);
+			}
+
+			MaterialSetOut.Materials.Add(GeneratedMID);
+		}
+
+	}
+
+}
+
+
+void UGenerateStaticMeshLODProcess::UpdateMaterialTextureParameters(
+	UMaterialInstanceDynamic* Material, 
+	const FSourceMaterialInfo& SourceMaterialInfo,
+	const TMap<UTexture2D*, UTexture2D*>& PreviewTextures, 
+	UTexture2D* PreviewNormalMap)
+{
+	ensure((SourceMaterialInfo.bHasNormalMap == false) || (PreviewNormalMap != nullptr));
+
+	Material->Modify();
+	int32 NumTextures = SourceMaterialInfo.SourceTextures.Num();
+	for (int32 ti = 0; ti < NumTextures; ++ti)
+	{
+		const FTextureInfo& SourceTex = SourceMaterialInfo.SourceTextures[ti];
+
+		if (SourceTex.Constraint == ETextureBakingConstraint::UseExistingTexture)
+		{
+			FMaterialParameterInfo ParamInfo(SourceTex.ParameterName);
+			Material->SetTextureParameterValueByInfo(ParamInfo, SourceTex.Texture);
+		}
+		else if (SourceTex.bIsNormalMap)
+		{
+			if (PreviewNormalMap)
+			{
+				FMaterialParameterInfo ParamInfo(SourceTex.ParameterName);
+				Material->SetTextureParameterValueByInfo(ParamInfo, PreviewNormalMap);
+			}
+		}
+		else if (SourceTex.bShouldBakeTexture)
+		{
+			if (CurrentSettings_Texture.bCombineTextures && SourceTex.bIsUsedInMultiTextureBaking)
+			{
+				continue;
+			}
+
+			UTexture2D*const* FoundTexture = PreviewTextures.Find(SourceTex.Texture);
+			if (ensure(FoundTexture))
+			{
+				FMaterialParameterInfo ParamInfo(SourceTex.ParameterName);
+				Material->SetTextureParameterValueByInfo(ParamInfo, *FoundTexture);
+			}
+		}
+	}
+	Material->PostEditChange();
+}
+
+
+
+
+
+bool UGenerateStaticMeshLODProcess::WriteDerivedAssetData()
+{
+	AllDerivedTextures.Reset();
+
+	constexpr bool bCreatingNewStaticMeshAsset = true;
+	DerivedAssetGUIDKey = FGuid::NewGuid().ToString(EGuidFormats::UniqueObjectGuid).ToUpper();
+
+	WriteDerivedTextures(bCreatingNewStaticMeshAsset);
+
+	WriteDerivedMaterials(bCreatingNewStaticMeshAsset);
+
+	WriteDerivedStaticMeshAsset();
+
+	// clear list of derived textures we were holding onto to prevent GC
+	AllDerivedTextures.Reset();
+
+	return true;
+}
+
+
+void UGenerateStaticMeshLODProcess::UpdateSourceAsset(bool bSetNewHDSourceAsset)
+{
+	AllDerivedTextures.Reset();
+
+	constexpr bool bCreatingNewStaticMeshAsset = false;
+	DerivedAssetGUIDKey = FGuid::NewGuid().ToString(EGuidFormats::UniqueObjectGuid).ToUpper();
+
+	WriteDerivedTextures(bCreatingNewStaticMeshAsset);
+
+	WriteDerivedMaterials(bCreatingNewStaticMeshAsset);
+
+	UpdateSourceStaticMeshAsset(bSetNewHDSourceAsset);
+
+	// clear list of derived textures we were holding onto to prevent GC
+	AllDerivedTextures.Reset();
+}
+
+
+
+bool UGenerateStaticMeshLODProcess::IsSourceAsset(const FString& AssetPath) const
+{
+	if (UEditorAssetLibrary::DoesAssetExist(AssetPath))
+	{
+		const FAssetData AssetData = UEditorAssetLibrary::FindAssetData(AssetPath);
+
+		for (const FSourceMaterialInfo& MaterialInfo : SourceMaterials)
+		{
+			UMaterialInterface* MaterialInterface = MaterialInfo.SourceMaterial.MaterialInterface;
+			if (MaterialInterface == nullptr)
+			{
+				continue;
+			}
+
+			FString SourceMaterialPath = UEditorAssetLibrary::GetPathNameForLoadedAsset(MaterialInterface);
+			if (UEditorAssetLibrary::FindAssetData(SourceMaterialPath) == AssetData)
+			{
+				return true;
+			}
+
+			for (const FTextureInfo& TextureInfo : MaterialInfo.SourceTextures)
+			{
+				FString SourceTexturePath = UEditorAssetLibrary::GetPathNameForLoadedAsset(TextureInfo.Texture);
+				if (UEditorAssetLibrary::FindAssetData(SourceTexturePath) == AssetData)
+				{
+					return true;
+				}
+			}
+		}
+	}
+
+	return false;
+}
+
+
+
+void UGenerateStaticMeshLODProcess::WriteDerivedTextures(bool bCreatingNewStaticMeshAsset)
+{
+	IAssetTools& AssetTools = FModuleManager::Get().LoadModuleChecked<FAssetToolsModule>("AssetTools").Get();
+
+	// this is a workaround for handling multiple materials that reference the same texture. Currently the code
+	// below will try to write that texture multiple times, which will fail when it tries to create a package
+	// for a filename that already exists
+	TMap<UTexture2D*, UTexture2D*> WrittenSourceToDerived;
+
+	// figure out which textures we might write have actually been referenced by a
+	// Material we are going to write. We do not want to write unreferenced textures.
+	bool bAnyMaterialUsedGeneratedNormalMap = false;
+	bool bAnyMaterialUsedMultiTexture = false;
+	int32 NumMaterials = SourceMaterials.Num();
+	for (int32 mi = 0; mi < NumMaterials; ++mi)
+	{
+		const FSourceMaterialInfo& SourceMaterialInfo = SourceMaterials[mi];
+		if (SourceMaterialInfo.bIsReusable
+			|| SourceMaterialInfo.bIsPreviouslyGeneratedMaterial
+			|| SourceMaterialInfo.Constraint == EMaterialBakingConstraint::UseExistingMaterial)
+		{
+			continue;
+		}
+		if (SourceMaterialInfo.bHasNormalMap)
+		{
+			bAnyMaterialUsedGeneratedNormalMap = true;
+		}
+		if (CurrentSettings_Texture.bCombineTextures && MultiTextureParameterName.Contains(mi))
+		{
+			bAnyMaterialUsedMultiTexture = true;
+		}
+	}
+
+
+	// write derived textures
+	check(DerivedMaterials.Num() == NumMaterials);
+	for (int32 mi = 0; mi < NumMaterials; ++mi)
+	{
+		const FSourceMaterialInfo& SourceMaterialInfo = SourceMaterials[mi];
+		if (SourceMaterialInfo.bIsReusable 
+			|| SourceMaterialInfo.bIsPreviouslyGeneratedMaterial
+			|| SourceMaterialInfo.Constraint == EMaterialBakingConstraint::UseExistingMaterial )
+		{
+			continue;
+		}
+
+		FDerivedMaterialInfo& DerivedMaterialInfo = DerivedMaterials[mi];
+
+		int32 NumTextures = SourceMaterialInfo.SourceTextures.Num();
+		check(DerivedMaterialInfo.DerivedTextures.Num() == NumTextures);
+		for (int32 ti = 0; ti < NumTextures; ++ti)
+		{
+			const FTextureInfo& SourceTex = SourceMaterialInfo.SourceTextures[ti];
+			FTextureInfo& DerivedTex = DerivedMaterialInfo.DerivedTextures[ti];
+
+			if (WrittenSourceToDerived.Contains(SourceTex.Texture))
+			{
+				// Already computed the derived texture from this source
+				DerivedTex.Texture = WrittenSourceToDerived[SourceTex.Texture];
+				continue;
+			}
+
+			bool bConvertToSRGB = SourceTex.Texture->SRGB;
+
+			if (DerivedTex.bShouldBakeTexture == false ||
+				DerivedTex.Constraint == ETextureBakingConstraint::UseExistingTexture)
+			{
+				continue;
+			}
+
+			FTexture2DBuilder TextureBuilder;
+			TextureBuilder.Initialize(FTexture2DBuilder::ETextureType::Color, DerivedTex.Dimensions);
+			TextureBuilder.GetTexture2D()->SRGB = bConvertToSRGB;
+			TextureBuilder.Copy(DerivedTex.Image, bConvertToSRGB);
+			TextureBuilder.Commit(false);
+
+			DerivedTex.Texture = TextureBuilder.GetTexture2D();
+			if (ensure(DerivedTex.Texture))
+			{
+				AllDerivedTextures.Add(DerivedTex.Texture);
+
+				FTexture2DBuilder::CopyPlatformDataToSourceData(DerivedTex.Texture, FTexture2DBuilder::ETextureType::Color);
+
+				// write asset
+				bool bWriteOK = WriteDerivedTexture(SourceTex.Texture, DerivedTex.Texture, bCreatingNewStaticMeshAsset);
+				ensure(bWriteOK);
+
+				WrittenSourceToDerived.Add(SourceTex.Texture, DerivedTex.Texture);
+			}
+		}
+
+	}
+
+
+	// write derived normal map
+	if (bAnyMaterialUsedGeneratedNormalMap) 
+	{
+		FTexture2DBuilder NormapMapBuilder;
+		NormapMapBuilder.Initialize(FTexture2DBuilder::ETextureType::NormalMap, DerivedNormalMapImage.Image.GetDimensions());
+		NormapMapBuilder.Copy(DerivedNormalMapImage.Image, false);
+		NormapMapBuilder.Commit(false);
+
+		DerivedNormalMapTex = NormapMapBuilder.GetTexture2D();
+		if (ensure(DerivedNormalMapTex))
+		{
+			FTexture2DBuilder::CopyPlatformDataToSourceData(DerivedNormalMapTex, FTexture2DBuilder::ETextureType::NormalMap);
+
+			// write asset
+			bool bWriteOK = WriteDerivedTexture(DerivedNormalMapTex, DerivedAssetNameNoSuffix + TEXT("_NormalMap"), bCreatingNewStaticMeshAsset);
+			ensure(bWriteOK);
+		}
+	}
+
+
+	// write multi-texture bake result
+	if (CurrentSettings_Texture.bCombineTextures && bAnyMaterialUsedMultiTexture)
+	{
+		bool bConvertToSRGB = true;
+		
+		FTexture2DBuilder MultiTextureBuilder;
+		MultiTextureBuilder.Initialize(FTexture2DBuilder::ETextureType::Color, DerivedMultiTextureBakeImage.Image.GetDimensions());
+		MultiTextureBuilder.GetTexture2D()->SRGB = bConvertToSRGB;
+		MultiTextureBuilder.Copy(DerivedMultiTextureBakeImage.Image, bConvertToSRGB);
+		MultiTextureBuilder.Commit(false);
+
+		DerivedMultiTextureBakeResult = MultiTextureBuilder.GetTexture2D();
+		if (ensure(DerivedMultiTextureBakeResult))
+		{
+			FTexture2DBuilder::CopyPlatformDataToSourceData(DerivedMultiTextureBakeResult, FTexture2DBuilder::ETextureType::Color);
+
+			// write asset
+			bool bWriteOK = WriteDerivedTexture(DerivedMultiTextureBakeResult, DerivedAssetNameNoSuffix + TEXT("_MultiTexture"), bCreatingNewStaticMeshAsset);
+			ensure(bWriteOK);
+		}
+	}
+	
+}
+
+
+
+bool UGenerateStaticMeshLODProcess::WriteDerivedTexture(UTexture2D* SourceTexture, UTexture2D* DerivedTexture, bool bCreatingNewStaticMeshAsset)
+{
+	IAssetTools& AssetTools = FModuleManager::Get().LoadModuleChecked<FAssetToolsModule>("AssetTools").Get();
+
+	FString SourceTexPath = UEditorAssetLibrary::GetPathNameForLoadedAsset(SourceTexture);
+	FString TexName = FPaths::GetBaseFilename(SourceTexPath, true);
+	return WriteDerivedTexture(DerivedTexture, TexName, bCreatingNewStaticMeshAsset);
+}
+
+
+
+bool UGenerateStaticMeshLODProcess::WriteDerivedTexture(UTexture2D* DerivedTexture, FString BaseTexName, bool bCreatingNewStaticMeshAsset)
+{
+	IAssetTools& AssetTools = FModuleManager::Get().LoadModuleChecked<FAssetToolsModule>("AssetTools").Get();
+
+	FString NewTexName = FString::Printf(TEXT("%s%s"), *BaseTexName, *DerivedSuffix);
+	FString NewAssetPath = FPaths::Combine(DerivedAssetFolder, NewTexName);
+
+	bool bNewAssetExistsInMemory = IsSourceAsset(NewAssetPath);
+
+	if (bCreatingNewStaticMeshAsset || bNewAssetExistsInMemory)
+	{
+		// Don't delete an existing asset. If name collision occurs, rename the new asset.
+		bool bNewAssetExists = UEditorAssetLibrary::DoesAssetExist(NewAssetPath);
+		while (bNewAssetExists)
+		{
+			GenerateStaticMeshLODProcessHelpers::AppendOrIncrementSuffix(NewTexName);
+			NewAssetPath = FPaths::Combine(DerivedAssetFolder, NewTexName);
+			bNewAssetExists = UEditorAssetLibrary::DoesAssetExist(NewAssetPath);
+		}
+	}
+	else
+	{
+		// Modifying the static mesh in place. Delete existing asset so that we can have a clean duplicate
+		bool bNewAssetExists = UEditorAssetLibrary::DoesAssetExist(NewAssetPath);
+		if (bNewAssetExists)
+		{
+			bool bDeleteOK = UEditorAssetLibrary::DeleteAsset(NewAssetPath);
+			ensure(bDeleteOK);
+		}
+	}
+
+	// create package
+	FString UniquePackageName, UniqueAssetName;
+	AssetTools.CreateUniqueAssetName(NewAssetPath, TEXT(""), UniquePackageName, UniqueAssetName);
+	UPackage* AssetPackage = CreatePackage(*UniquePackageName);
+	check(AssetPackage);
+
+	// move texture from Transient package to new package
+	DerivedTexture->Rename(*UniqueAssetName, AssetPackage, REN_None);
+	// remove transient flag, add public/standalone/transactional
+	DerivedTexture->ClearFlags(RF_Transient);
+	DerivedTexture->SetFlags(RF_Public | RF_Standalone | RF_Transactional);
+	// do we need to Modify() it? we are not doing any undo/redo
+	DerivedTexture->Modify();
+
+	// set metadata tag so we can identify this material later
+	DerivedTexture->GetOutermost()->GetMetaData()->SetValue(DerivedTexture, TEXT("StaticMeshLOD.IsBakedTexture"), TEXT("true"));
+	DerivedTexture->GetOutermost()->GetMetaData()->SetValue(DerivedTexture, TEXT("StaticMeshLOD.SourceAssetPath"), *GetSourceAssetPath());
+	DerivedTexture->GetOutermost()->GetMetaData()->SetValue(DerivedTexture, TEXT("StaticMeshLOD.GenerationGUID"), *DerivedAssetGUIDKey);
+
+	DerivedTexture->UpdateResource();
+	DerivedTexture->PostEditChange();		// this may be necessary if any Materials are using this texture
+	DerivedTexture->MarkPackageDirty();
+
+	FAssetRegistryModule::AssetCreated(DerivedTexture);		// necessary?
+
+	return true;
+}
+
+
+void UGenerateStaticMeshLODProcess::WriteDerivedMaterials(bool bCreatingNewStaticMeshAsset)
+{
+	IAssetTools& AssetTools = FModuleManager::Get().LoadModuleChecked<FAssetToolsModule>("AssetTools").Get();
+
+	int32 NumMaterials = SourceMaterials.Num();
+	check(DerivedMaterials.Num() == NumMaterials);
+	for (int32 mi = 0; mi < NumMaterials; ++mi)
+	{
+		const FSourceMaterialInfo& SourceMaterialInfo = SourceMaterials[mi];
+		if (SourceMaterialInfo.bIsReusable 
+			|| SourceMaterialInfo.bIsPreviouslyGeneratedMaterial
+			|| SourceMaterialInfo.Constraint == EMaterialBakingConstraint::UseExistingMaterial )
+		{
+			continue;
+		}
+
+		FDerivedMaterialInfo& DerivedMaterialInfo = DerivedMaterials[mi];
+
+		UMaterialInterface* MaterialInterface = SourceMaterialInfo.SourceMaterial.MaterialInterface;
+		if (MaterialInterface == nullptr)
+		{
+			DerivedMaterialInfo.DerivedMaterial.MaterialInterface = nullptr;
+			continue;
+		}
+		bool bSourceIsMIC = (Cast<UMaterialInstanceConstant>(MaterialInterface) != nullptr);
+
+		FString SourceMaterialPath = UEditorAssetLibrary::GetPathNameForLoadedAsset(MaterialInterface);
+		FString MaterialName = FPaths::GetBaseFilename(SourceMaterialPath, true);
+		FString NewMaterialName = FString::Printf(TEXT("%s%s"), *MaterialName, *DerivedSuffix);
+		FString NewMaterialPath = FPaths::Combine(DerivedAssetFolder, NewMaterialName);
+		bool bNewAssetExistsInMemory = IsSourceAsset(NewMaterialPath);
+
+		if (bCreatingNewStaticMeshAsset || bNewAssetExistsInMemory)
+		{
+			// Don't delete an existing material. If name collision occurs, rename the new material.
+			bool bNewAssetExists = UEditorAssetLibrary::DoesAssetExist(NewMaterialPath);
+			while (bNewAssetExists)
+			{
+				GenerateStaticMeshLODProcessHelpers::AppendOrIncrementSuffix(NewMaterialName);
+				NewMaterialPath = FPaths::Combine(DerivedAssetFolder, NewMaterialName);
+				bNewAssetExists = UEditorAssetLibrary::DoesAssetExist(NewMaterialPath);
+			}
+		}
+		else
+		{
+			// Modifying the static mesh in place. Delete existing asset so that we can have a clean duplicate
+			bool bNewAssetExists = UEditorAssetLibrary::DoesAssetExist(NewMaterialPath);
+			if (bNewAssetExists)
+			{
+				bool bDeleteOK = UEditorAssetLibrary::DeleteAsset(NewMaterialPath);
+				ensure(bDeleteOK);
+			}
+		}
+
+		// If source is a MIC, we can just duplicate it. If it is a UMaterial, we want to
+		// create a child MIC? Or we could dupe the Material and rewrite the textures.
+		// Probably needs to be an option.
+		UMaterialInstanceConstant* GeneratedMIC = nullptr;
+		if (bSourceIsMIC)
+		{
+			UObject* DupeAsset = UEditorAssetLibrary::DuplicateAsset(SourceMaterialPath, NewMaterialPath);
+			GeneratedMIC = Cast<UMaterialInstanceConstant>(DupeAsset);
+		}
+		else
+		{
+			UMaterial* SourceMaterial = MaterialInterface->GetBaseMaterial();
+			if (ensure(SourceMaterial))
+			{
+				UMaterialInstanceConstantFactoryNew* Factory = NewObject<UMaterialInstanceConstantFactoryNew>();
+				Factory->InitialParent = SourceMaterial;
+
+				UObject* NewAsset = AssetTools.CreateAsset(NewMaterialName, FPackageName::GetLongPackagePath(NewMaterialPath),
+					UMaterialInstanceConstant::StaticClass(), Factory);
+
+				GeneratedMIC = Cast<UMaterialInstanceConstant>(NewAsset);
+			}
+		}
+
+		if (GeneratedMIC != nullptr)
+		{
+			// set metadata tag so we can identify this material later
+			GeneratedMIC->GetOutermost()->GetMetaData()->SetValue(GeneratedMIC, TEXT("StaticMeshLOD.IsGeneratedMaterial"), TEXT("true"));
+			GeneratedMIC->GetOutermost()->GetMetaData()->SetValue(GeneratedMIC, TEXT("StaticMeshLOD.SourceAssetPath"), *GetSourceAssetPath());
+			GeneratedMIC->GetOutermost()->GetMetaData()->SetValue(GeneratedMIC, TEXT("StaticMeshLOD.GenerationGUID"), *DerivedAssetGUIDKey);
+
+			// rewrite texture parameters to new textures
+			UpdateMaterialTextureParameters(GeneratedMIC, DerivedMaterialInfo);
+		}
+
+		// update StaticMaterial
+		DerivedMaterialInfo.DerivedMaterial.MaterialInterface = GeneratedMIC;
+		DerivedMaterialInfo.DerivedMaterial.MaterialSlotName = FName(FString::Printf(TEXT("GeneratedMat%d"), mi));
+		DerivedMaterialInfo.DerivedMaterial.ImportedMaterialSlotName = DerivedMaterialInfo.DerivedMaterial.MaterialSlotName;
+
+		if (GeneratedMIC != nullptr)
+		{
+			if (CurrentSettings_Texture.bCombineTextures && DerivedMultiTextureBakeResult && MultiTextureParameterName.Contains(mi))
+			{
+				FMaterialParameterInfo ParamInfo(MultiTextureParameterName[mi]);
+				GeneratedMIC->SetTextureParameterValueEditorOnly(ParamInfo, DerivedMultiTextureBakeResult);
+			}
+		}
+
+	}
+}
+
+
+
+void UGenerateStaticMeshLODProcess::UpdateMaterialTextureParameters(UMaterialInstanceConstant* Material, FDerivedMaterialInfo& DerivedMaterialInfo)
+{
+	Material->Modify();
+
+	int32 NumTextures = DerivedMaterialInfo.DerivedTextures.Num();
+	for (int32 ti = 0; ti < NumTextures; ++ti)
+	{
+		FTextureInfo& DerivedTex = DerivedMaterialInfo.DerivedTextures[ti];
+
+		if (DerivedTex.Constraint == ETextureBakingConstraint::UseExistingTexture)
+		{
+			UTexture2D* SourceTexture = SourceMaterials[DerivedMaterialInfo.SourceMaterialIndex].SourceTextures[ti].Texture;
+			FMaterialParameterInfo ParamInfo(DerivedTex.ParameterName);
+			Material->SetTextureParameterValueEditorOnly(ParamInfo, SourceTexture);
+		}
+		else if (DerivedTex.bIsNormalMap)
+		{
+			if (ensure(DerivedNormalMapTex))
+			{
+				FMaterialParameterInfo ParamInfo(DerivedTex.ParameterName);
+				Material->SetTextureParameterValueEditorOnly(ParamInfo, DerivedNormalMapTex);
+			}
+		}
+		else if (DerivedTex.bShouldBakeTexture)
+		{
+			UTexture2D* NewTexture = DerivedTex.Texture;
+			if (ensure(NewTexture))
+			{
+				FMaterialParameterInfo ParamInfo(DerivedTex.ParameterName);
+				Material->SetTextureParameterValueEditorOnly(ParamInfo, NewTexture);
+			}
+		}
+	}
+
+	Material->PostEditChange();
+}
+
+
+
+void UGenerateStaticMeshLODProcess::WriteDerivedStaticMeshAsset()
+{
+	// [TODO] should we try to re-use existing asset here, or should we delete it? 
+	// The source asset might have had any number of config changes that we want to
+	// preserve in the duplicate...
+	UStaticMesh* GeneratedStaticMesh = nullptr;
+	if (UEditorAssetLibrary::DoesAssetExist(DerivedAssetPath))
+	{
+		UObject* LoadedAsset = UEditorAssetLibrary::LoadAsset(DerivedAssetPath);
+		GeneratedStaticMesh = Cast<UStaticMesh>(LoadedAsset);
+	}
+	else
+	{
+		UObject* DupeAsset = UEditorAssetLibrary::DuplicateAsset(SourceAssetPath, DerivedAssetPath);
+		GeneratedStaticMesh = Cast<UStaticMesh>(DupeAsset);
+	}
+
+	// DuplicateAsset can fail in the event of a dead package from an Asset Reload.
+	// This edge case is currently prevented by the GenerateStaticMeshLODAssetToolBuilder.
+	// [TODO] Handle a failed DuplicateAsset case by creating a new StaticMesh asset.
+	if (!ensure(GeneratedStaticMesh))
+	{
+		return;
+	}
+
+	// make sure transactional flag is on
+	GeneratedStaticMesh->SetFlags(RF_Transactional);
+	GeneratedStaticMesh->Modify();
+
+	// update MeshDescription LOD0 mesh
+	GeneratedStaticMesh->SetNumSourceModels(1);
+	FMeshDescription* MeshDescription = GeneratedStaticMesh->GetMeshDescription(0);
+	FConversionToMeshDescriptionOptions ConversionOptions;
+	FDynamicMeshToMeshDescription Converter(ConversionOptions);
+	Converter.Convert(&DerivedLODMesh, *MeshDescription);
+	GeneratedStaticMesh->CommitMeshDescription(0);
+
+	// construct new material slots list
+	TArray<FStaticMaterial> NewMaterials;
+	int32 NumMaterials = SourceMaterials.Num();
+	for (int32 mi = 0; mi < NumMaterials; ++mi)
+	{
+		if (!SourceMaterials[mi].bIsPreviouslyGeneratedMaterial)	// Skip previously generated
+		{
+			if (SourceMaterials[mi].bIsReusable || SourceMaterials[mi].Constraint == EMaterialBakingConstraint::UseExistingMaterial )
+			{
+				NewMaterials.Add(SourceMaterials[mi].SourceMaterial);
+			}
+			else
+			{
+				NewMaterials.Add(DerivedMaterials[mi].DerivedMaterial);
+			}
+		}
+	}
+
+	// update materials on generated mesh
+	GeneratedStaticMesh->SetStaticMaterials(NewMaterials);
+
+	// collision
+	FPhysicsDataCollection NewCollisionGeo;
+	NewCollisionGeo.Geometry = DerivedCollision;
+	NewCollisionGeo.CopyGeometryToAggregate();
+
+	// code below derived from FStaticMeshEditor::DuplicateSelectedPrims()
+	UBodySetup* BodySetup = GeneratedStaticMesh->GetBodySetup();
+	// mark the BodySetup for modification. Do we need to modify the UStaticMesh??
+	BodySetup->Modify();
+	//Clear the cache (PIE may have created some data), create new GUID    (comment from StaticMeshEditor)
+	BodySetup->InvalidatePhysicsData();
+	BodySetup->RemoveSimpleCollision();
+	BodySetup->AggGeom = NewCollisionGeo.AggGeom;
+	// update collision type
+	BodySetup->CollisionTraceFlag = ECollisionTraceFlag::CTF_UseDefault;
+	// rebuild physics data
+	BodySetup->InvalidatePhysicsData();
+	BodySetup->CreatePhysicsMeshes();
+
+	// do we need to do a post edit change here??
+
+	GeneratedStaticMesh->RecreateNavCollision();
+
+	GeneratedStaticMesh->GetOutermost()->GetMetaData()->SetValue(GeneratedStaticMesh, TEXT("StaticMeshLOD.IsGeneratedMesh"), TEXT("true"));
+	GeneratedStaticMesh->GetOutermost()->GetMetaData()->SetValue(GeneratedStaticMesh, TEXT("StaticMeshLOD.SourceAssetPath"), *GetSourceAssetPath());
+	GeneratedStaticMesh->GetOutermost()->GetMetaData()->SetValue(GeneratedStaticMesh, TEXT("StaticMeshLOD.GenerationGUID"), *DerivedAssetGUIDKey);
+
+	// done updating mesh
+	GeneratedStaticMesh->PostEditChange();
+}
+
+
+
+void UGenerateStaticMeshLODProcess::UpdateSourceStaticMeshAsset(bool bSetNewHDSourceAsset)
+{
+	GEditor->BeginTransaction(LOCTEXT("UpdateExistingAssetMessage", "Added Generated LOD"));
+
+	SourceStaticMesh->Modify();
+
+	FStaticMeshSourceModel& SrcModel = SourceStaticMesh->GetSourceModel(0);
+	SourceStaticMesh->ModifyMeshDescription(0);
+
+	// if we want to save the input high-poly asset as the HiResSource, do that here
+	if (bSetNewHDSourceAsset && bUsingHiResSource == false)
+	{
+		SourceStaticMesh->ModifyHiResMeshDescription();
+
+		FMeshDescription* NewSourceMD = SourceStaticMesh->CreateHiResMeshDescription();
+		*NewSourceMD = *SourceMeshDescription;		// todo: can MoveTemp here, we don't need this memory anymore??
+
+		FStaticMeshSourceModel& HiResSrcModel = SourceStaticMesh->GetHiResSourceModel();
+		// Generally copy LOD0 build settings, although many of these will be ignored
+		HiResSrcModel.BuildSettings = SrcModel.BuildSettings;
+		// on the HiRes we store the existing normals and tangents, which we already auto-computed if necessary
+		HiResSrcModel.BuildSettings.bRecomputeNormals = false;
+		HiResSrcModel.BuildSettings.bRecomputeTangents = false;
+		// TODO: what should we do about Lightmap UVs?
+
+		SourceStaticMesh->CommitHiResMeshDescription();
+	}
+
+	// Next bit is tricky, we have to build the final StaticMesh Material Set.
+	// We have the existing Source materials we want to keep, except if some
+	// were identified as being auto-generated by a previous run of AutoLOD, we want to leave those out.
+	// Then we want to add any New generated materials. 
+	// The main complication is that we cannot change the slot indices for the existing Source materials,
+	// as we would have to fix up the HighRes source. Ideally they are the first N slots, and we
+	// just append the new ones. But we cannot guarantee this, so if there are gaps we will interleave
+	// the new materials when possible.
+
+
+	// Figure out which derived materials we are actually going to use.
+	// We will not use materials we can re-use from the source, or materials that were previously auto-generated.
+	int32 NumMaterials = SourceMaterials.Num();
+	TArray<int32> DerivedMatSlotIndexMap;		// this maps from current DerivedMesh slot indices to their final slot indices
+	DerivedMatSlotIndexMap.SetNum(NumMaterials);
+	TArray<int32> DerivedMaterialsToAdd;		// list of derived material indices we need to store in the final material set
+	int32 NewMaterialCount = 0;
+	for (int32 mi = 0; mi < NumMaterials; ++mi)
+	{
+		if (SourceMaterials[mi].bIsPreviouslyGeneratedMaterial)
+		{
+			DerivedMatSlotIndexMap[mi] = -2;		// these materials do not appear in DerivedMesh or SourceMesh and should be skipped/discarded (todo: and deleted?)
+		}
+		else  if (SourceMaterials[mi].bIsReusable
+				  || SourceMaterials[mi].Constraint == EMaterialBakingConstraint::UseExistingMaterial)	
+		{
+			DerivedMatSlotIndexMap[mi] = mi;		// if we can re-use existing material we just rewrite to existing material slot index
+		}
+		else
+		{
+			DerivedMatSlotIndexMap[mi] = -1;		// will need to allocate a new slot for this material
+			DerivedMaterialsToAdd.Add(mi);
+			NewMaterialCount++;
+		}
+	}
+	int32 CurRemainingDerivedIdx = 0;
+
+	// Copy existing materials we want to keep to new StaticMesh materials set. 
+	// If there are any gaps left by skipping previously-derived materials, try to 
+	// tuck in a new derived material that is waiting to be allocated to a slot
+	TArray<FStaticMaterial> NewMaterialSet;
+	TArray<int32> DerivedMaterialSlotIndices;
+	for (int32 k = 0; k < SourceMaterials.Num(); ++k)
+	{
+		if (SourceMaterials[k].bIsPreviouslyGeneratedMaterial == false)
+		{
+			NewMaterialSet.Add(SourceMaterials[k].SourceMaterial);
+		}
+		else if (CurRemainingDerivedIdx < DerivedMaterialsToAdd.Num() )
+		{
+			int32 DerivedIdx = DerivedMaterialsToAdd[CurRemainingDerivedIdx];
+			CurRemainingDerivedIdx++;
+			DerivedMatSlotIndexMap[DerivedIdx] = NewMaterialSet.Num();
+			DerivedMaterialSlotIndices.Add(DerivedMatSlotIndexMap[DerivedIdx]);
+			NewMaterialSet.Add(DerivedMaterials[DerivedIdx].DerivedMaterial);
+		}
+		else
+		{
+			// we ran out of new materials to allocate and so just add empty ones??
+			ensure(false);
+			NewMaterialSet.Add(FStaticMaterial());
+		}
+	}
+
+	// if we have any new derived materials left, append them to the material set
+	while (CurRemainingDerivedIdx < DerivedMaterialsToAdd.Num())
+	{
+		int32 DerivedIdx = DerivedMaterialsToAdd[CurRemainingDerivedIdx];
+		CurRemainingDerivedIdx++;
+		DerivedMatSlotIndexMap[DerivedIdx] = NewMaterialSet.Num();
+		DerivedMaterialSlotIndices.Add(DerivedMatSlotIndexMap[DerivedIdx]);
+		NewMaterialSet.Add(DerivedMaterials[DerivedIdx].DerivedMaterial);
+	}
+
+	// apply the material slot index rewrite map to the DerivedMesh
+	DerivedLODMesh.Attributes()->EnableMaterialID();
+	FDynamicMeshMaterialAttribute* MaterialIDs = DerivedLODMesh.Attributes()->GetMaterialID();
+	for (int32 tid : DerivedLODMesh.TriangleIndicesItr())
+	{
+		int32 CurMaterialID = MaterialIDs->GetValue(tid);
+		int32 NewMaterialID = DerivedMatSlotIndexMap[CurMaterialID];
+		if (ensure(NewMaterialID >= 0))
+		{
+			MaterialIDs->SetValue(tid, NewMaterialID);
+		}
+	}
+
+	// update materials on generated mesh
+	SourceStaticMesh->SetStaticMaterials(NewMaterialSet);
+
+	// store new derived LOD as LOD 0
+	SourceStaticMesh->SetNumSourceModels(1);
+	FMeshDescription* MeshDescription = SourceStaticMesh->GetMeshDescription(0);
+	if (MeshDescription == nullptr)
+	{
+		MeshDescription = SourceStaticMesh->CreateMeshDescription(0);
+	}
+	FConversionToMeshDescriptionOptions ConversionOptions;
+	FDynamicMeshToMeshDescription Converter(ConversionOptions);
+	Converter.Convert(&DerivedLODMesh, *MeshDescription);
+
+	// calculate tangents
+	Converter.UpdateTangents(&DerivedLODMesh, *MeshDescription, &DerivedLODMeshTangents);
+
+	// set slot names on MeshDescription to match those we set on the generated FStaticMaterials, 
+	// because StaticMesh RenderBuffer setup will do matching-name lookups and if it is NAME_None
+	// we will get the wrong material!
+	FStaticMeshAttributes Attributes(*MeshDescription);
+	TPolygonGroupAttributesRef<FName> PolygonGroupImportedMaterialSlotNames = Attributes.GetPolygonGroupMaterialSlotNames();
+	for (int32 SlotIdx : DerivedMaterialSlotIndices)
+	{
+		// It's possible that NewMaterialSet.Num() > PolygonGroupImportedMaterialSlotNames.GetNumElements()
+		// if there are new materials that aren't referenced by any triangles...
+		if (SlotIdx < PolygonGroupImportedMaterialSlotNames.GetNumElements())		
+		{
+			PolygonGroupImportedMaterialSlotNames.Set(SlotIdx, NewMaterialSet[SlotIdx].ImportedMaterialSlotName);
+		}
+	}
+
+	// Disable auto-generated normals/tangents, we need to use the ones we computed in LOD Generator
+	SrcModel.BuildSettings.bRecomputeNormals = false;
+	SrcModel.BuildSettings.bRecomputeTangents = false;
+
+	// this will prevent simplification?
+	SrcModel.ReductionSettings.MaxDeviation = 0.0f;
+	SrcModel.ReductionSettings.PercentTriangles = 1.0f;
+	SrcModel.ReductionSettings.PercentVertices = 1.0f;
+
+	// commit update
+	SourceStaticMesh->CommitMeshDescription(0);
+
+	// collision
+	FPhysicsDataCollection NewCollisionGeo;
+	NewCollisionGeo.Geometry = DerivedCollision;
+	NewCollisionGeo.CopyGeometryToAggregate();
+
+	// code below derived from FStaticMeshEditor::DuplicateSelectedPrims()
+	UBodySetup* BodySetup = SourceStaticMesh->GetBodySetup();
+	// mark the BodySetup for modification. Do we need to modify the UStaticMesh??
+	BodySetup->Modify();
+	//Clear the cache (PIE may have created some data), create new GUID    (comment from StaticMeshEditor)
+	BodySetup->InvalidatePhysicsData();
+	BodySetup->RemoveSimpleCollision();
+	BodySetup->AggGeom = NewCollisionGeo.AggGeom;
+	// update collision type
+	BodySetup->CollisionTraceFlag = ECollisionTraceFlag::CTF_UseDefault;
+	// rebuild physics data
+	BodySetup->InvalidatePhysicsData();
+	BodySetup->CreatePhysicsMeshes();
+
+	// do we need to do a post edit change here??
+
+	SourceStaticMesh->RecreateNavCollision();
+
+	// save UUID used for generated Assets
+	SourceStaticMesh->GetOutermost()->GetMetaData()->SetValue(SourceStaticMesh, TEXT("StaticMeshLOD.GenerationGUID"), *DerivedAssetGUIDKey);
+
+	GEditor->EndTransaction();
+
+	// done updating mesh
+	SourceStaticMesh->PostEditChange();
+}
+
+
+
+#undef LOCTEXT_NAMESPACE
diff --git a/Engine/Plugins/Editor/MeshLODToolset/Source/MeshLODToolset/Private/MeshLODToolsetModule.cpp b/Engine/Plugins/Editor/MeshLODToolset/Source/MeshLODToolset/Private/MeshLODToolsetModule.cpp
new file mode 100644
--- /dev/null
+++ b/Engine/Plugins/Editor/MeshLODToolset/Source/MeshLODToolset/Private/MeshLODToolsetModule.cpp
@@ -0,0 +1,48 @@
+// Copyright Epic Games, Inc. All Rights Reserved.
+
+#include "MeshLODToolsetModule.h"
+#include "Misc/CoreDelegates.h"
+#include "Tools/LODGenerationSettingsAsset.h"
+#include "PropertyEditorModule.h"
+#include "Tools/DetailsCustomizations/AutoLODToolCustomizations.h"
+#include "Tools/GenerateStaticMeshLODAssetTool.h"
+#include "Modules/ModuleManager.h"
+
+#define LOCTEXT_NAMESPACE "FMeshLODToolsetModule"
+
+DEFINE_LOG_CATEGORY(LogMeshLODToolset);
+
+// This function will execute after your module is loaded into memory; the exact timing is specified in the .uplugin file per-module
+void FMeshLODToolsetModule::StartupModule()
+{
+	FCoreDelegates::OnPostEngineInit.AddRaw(this, &FMeshLODToolsetModule::OnPostEngineInit);
+}
+
+// This function may be called during shutdown to clean up your module.  For modules that support dynamic reloading,
+// we call this function before unloading the module.
+void FMeshLODToolsetModule::ShutdownModule()
+{
+	// Unregister customizations
+	FPropertyEditorModule* PropertyEditorModule = FModuleManager::GetModulePtr<FPropertyEditorModule>("PropertyEditor");
+	if (PropertyEditorModule)
+	{
+		for (FName ClassName : ClassesToUnregisterOnShutdown)
+		{
+			PropertyEditorModule->UnregisterCustomClassLayout(ClassName);
+		}
+	}
+}
+
+void FMeshLODToolsetModule::OnPostEngineInit()
+{
+	ClassesToUnregisterOnShutdown.Reset();
+
+	FPropertyEditorModule& PropertyModule = FModuleManager::LoadModuleChecked<FPropertyEditorModule>("PropertyEditor");
+
+	PropertyModule.RegisterCustomClassLayout("GenerateStaticMeshLODAssetToolProperties", FOnGetDetailCustomizationInstance::CreateStatic(&FAutoLODToolDetails::MakeInstance));
+	ClassesToUnregisterOnShutdown.Add(UGenerateStaticMeshLODAssetToolProperties::StaticClass()->GetFName());
+}
+
+#undef LOCTEXT_NAMESPACE
+	
+IMPLEMENT_MODULE(FMeshLODToolsetModule, MeshLODToolset)
diff --git a/Engine/Plugins/Editor/MeshLODToolset/Source/MeshLODToolset/Private/Tools/DetailsCustomizations/AutoLODToolCustomizations.cpp b/Engine/Plugins/Editor/MeshLODToolset/Source/MeshLODToolset/Private/Tools/DetailsCustomizations/AutoLODToolCustomizations.cpp
new file mode 100644
--- /dev/null
+++ b/Engine/Plugins/Editor/MeshLODToolset/Source/MeshLODToolset/Private/Tools/DetailsCustomizations/AutoLODToolCustomizations.cpp
@@ -0,0 +1,80 @@
+// Copyright Epic Games, Inc. All Rights Reserved.
+
+#include "Tools/DetailsCustomizations/AutoLODToolCustomizations.h"
+#include "DetailLayoutBuilder.h"
+#include "DetailWidgetRow.h"
+#include "IDetailGroup.h"
+#include "Widgets/Input/SEditableTextBox.h"
+#include "Tools/GenerateStaticMeshLODAssetTool.h"
+
+#define LOCTEXT_NAMESPACE "AutoLODToolCustomizations"
+
+TSharedRef<IDetailCustomization> FAutoLODToolDetails::MakeInstance()
+{
+	return MakeShareable(new FAutoLODToolDetails);
+}
+
+void FAutoLODToolDetails::CustomizeDetails(IDetailLayoutBuilder& DetailBuilder)
+{
+	TSharedRef<IPropertyHandle> PreprocessingGroupSettings = DetailBuilder.GetProperty(GET_MEMBER_NAME_CHECKED(UGenerateStaticMeshLODAssetToolProperties, Preprocessing), 
+		UGenerateStaticMeshLODAssetToolProperties::StaticClass());
+	
+	// NOTE: We are completely rebuilding the widgets for the FGenerateStaticMeshLODProcess_PreprocessSettings struct. If the struct changes, this code will
+	// need to be updated.
+	 	
+	// Hide existing group
+	DetailBuilder.HideProperty(PreprocessingGroupSettings);
+
+	// Create new group
+	IDetailCategoryBuilder& ConfigCategory = DetailBuilder.EditCategory("Generator Configuration");
+
+	IDetailGroup& DetailGroup = ConfigCategory.AddGroup("Preprocessing", LOCTEXT("Preprocessing", "Preprocessing"));
+
+	TSharedPtr<IPropertyHandle> FilterGroupLayerHandle = PreprocessingGroupSettings->GetChildHandle(FName("FilterGroupLayer"));
+	check(FilterGroupLayerHandle);
+
+	DetailGroup.AddWidgetRow()
+	.NameContent()
+	[
+		FilterGroupLayerHandle->CreatePropertyNameWidget(LOCTEXT("DetailFilterGroupLayerName", "Detail Filter Group Layer"))
+	]
+	.ValueContent()
+	.VAlign(VAlign_Center)
+	[
+		SNew(SEditableTextBox)
+		.Text(FText())
+		.HintText(LOCTEXT("DetailFilterGroupLayerValueHintText", "(e.g. PreFilterGroups)"))
+		.OnTextCommitted_Lambda([FilterGroupLayerHandle](const FText& Val, ETextCommit::Type TextCommitType)
+		{
+			FilterGroupLayerHandle->SetValue(Val.ToString());
+		})
+	];
+
+	TSharedPtr<IPropertyHandle> ThickenWeightMapNameHandle = PreprocessingGroupSettings->GetChildHandle(FName("ThickenWeightMapName"));
+	check(ThickenWeightMapNameHandle);
+
+	DetailGroup.AddWidgetRow()
+	.NameContent()
+	[
+		ThickenWeightMapNameHandle->CreatePropertyNameWidget(LOCTEXT("ThickenWeightMapName", "Thicken Weight Map"))
+	]
+	.ValueContent()
+	.VAlign(VAlign_Center)
+	[
+		SNew(SEditableTextBox)
+		.Text(FText())
+		.HintText(LOCTEXT("ThickenWeightMapNameValueHintText", "(e.g. ThickenMap)"))
+		.OnTextCommitted_Lambda([ThickenWeightMapNameHandle](const FText& Val, ETextCommit::Type TextCommitType)
+		{
+			ThickenWeightMapNameHandle->SetValue(Val.ToString());
+		})
+	];
+
+	TSharedPtr<IPropertyHandle> ThickenAmountHandle = PreprocessingGroupSettings->GetChildHandle(FName("ThickenAmount"));
+	check(ThickenAmountHandle);
+
+	DetailGroup.AddPropertyRow(ThickenAmountHandle.ToSharedRef());
+}
+
+#undef LOCTEXT_NAMESPACE
+
diff --git a/Engine/Plugins/Editor/MeshLODToolset/Source/MeshLODToolset/Private/Tools/GenerateStaticMeshLODAssetTool.cpp b/Engine/Plugins/Editor/MeshLODToolset/Source/MeshLODToolset/Private/Tools/GenerateStaticMeshLODAssetTool.cpp
new file mode 100644
--- /dev/null
+++ b/Engine/Plugins/Editor/MeshLODToolset/Source/MeshLODToolset/Private/Tools/GenerateStaticMeshLODAssetTool.cpp
@@ -0,0 +1,731 @@
+// Copyright Epic Games, Inc. All Rights Reserved.
+
+#include "Tools/GenerateStaticMeshLODAssetTool.h"
+#include "Drawing/LineSetComponent.h"
+#include "InteractiveToolManager.h"
+
+#include "Drawing/PreviewGeometryActor.h"
+#include "ToolSetupUtil.h"
+
+
+#include "Materials/MaterialInterface.h"
+#include "Util/ColorConstants.h"
+
+
+
+#include "Physics/CollisionPropertySets.h"
+#include "Physics/PhysicsDataCollection.h"
+#include "Physics/CollisionGeometryVisualization.h"
+
+#include "Components/StaticMeshComponent.h"
+
+#include "Misc/Paths.h"
+#include "EditorAssetLibrary.h"
+#include "ObjectTools.h"
+
+#include "ModelingOperators.h"
+#include "MeshOpPreviewHelpers.h"
+#include "Generators/GridBoxMeshGenerator.h"
+
+#include "Framework/Notifications/NotificationManager.h"
+#include "Framework/Docking/TabManager.h"
+#include "TargetInterfaces/MaterialProvider.h"
+#include "Widgets/Notifications/SNotificationList.h"
+#include "TargetInterfaces/PrimitiveComponentBackedTarget.h"
+#include "TargetInterfaces/StaticMeshBackedTarget.h"
+#include "ToolContextInterfaces.h"
+#include "ToolTargetManager.h"
+#include "ModelingToolTargetUtil.h"
+
+#include "Tools/LODGenerationSettingsAsset.h"
+
+static_assert(WITH_EDITOR, "Tool being compiled without editor");
+#include "Misc/ScopedSlowTask.h"
+
+#include UE_INLINE_GENERATED_CPP_BY_NAME(GenerateStaticMeshLODAssetTool)
+
+using namespace UE::Geometry;
+
+#define LOCTEXT_NAMESPACE "UGenerateStaticMeshLODAssetTool"
+
+//
+// Local Op Stuff
+//
+
+namespace GenerateStaticMeshLODAssetLocals
+{
+
+	class FGenerateStaticMeshLODAssetOperatorOp : public FDynamicMeshOperator, public FGCObject
+	{
+	public:
+
+		// Inputs
+		TObjectPtr<UGenerateStaticMeshLODProcess> GenerateProcess;
+		FGenerateStaticMeshLODProcess_PreprocessSettings GeneratorSettings_Preprocess;
+		FGenerateStaticMeshLODProcessSettings GeneratorSettings_MeshGeneration;
+		FGenerateStaticMeshLODProcess_SimplifySettings GeneratorSettings_Simplify;
+		FGenerateStaticMeshLODProcess_NormalsSettings GeneratorSettings_Normals;
+		FGenerateStaticMeshLODProcess_TextureSettings GeneratorSettings_Texture;
+		FGenerateStaticMeshLODProcess_UVSettings GeneratorSettings_UV;
+		FGenerateStaticMeshLODProcess_CollisionSettings GeneratorSettings_Collision;
+		TArray<TPair<UMaterialInterface*, UGenerateStaticMeshLODProcess::EMaterialBakingConstraint>> MaterialConstraints;
+		TArray<TPair<UTexture2D*, UGenerateStaticMeshLODProcess::ETextureBakingConstraint>> TextureConstraints;
+		
+		// Outputs
+		// Inherited: 	TUniquePtr<FDynamicMesh3> ResultMesh;
+		// 				FTransform3d ResultTransform;
+		FMeshTangentsd ResultTangents;
+		FSimpleShapeSet3d ResultCollision;
+
+		void AddReferencedObjects(FReferenceCollector& Collector) override
+		{
+			Collector.AddReferencedObject(GenerateProcess);
+		}
+		virtual FString GetReferencerName() const override
+		{
+			return TEXT("GenerateStaticMeshLODAssetLocals::FGenerateStaticMeshLODAssetOperatorOp");
+		}
+
+		void CalculateResult(FProgressCancel* Progress) override
+		{
+			auto DoCompute = [this](FProgressCancel* Progress)		// bracket this computation with lock/unlock
+			{
+				if (Progress && Progress->Cancelled())
+				{
+					return;
+				}
+
+				GenerateProcess->UpdatePreprocessSettings(GeneratorSettings_Preprocess);
+				GenerateProcess->UpdateSettings(GeneratorSettings_MeshGeneration);
+				GenerateProcess->UpdateSimplifySettings(GeneratorSettings_Simplify);
+				GenerateProcess->UpdateNormalsSettings(GeneratorSettings_Normals);
+				GenerateProcess->UpdateTextureSettings(GeneratorSettings_Texture);
+				GenerateProcess->UpdateUVSettings(GeneratorSettings_UV);
+				GenerateProcess->UpdateCollisionSettings(GeneratorSettings_Collision);
+				for (TPair<UMaterialInterface*, UGenerateStaticMeshLODProcess::EMaterialBakingConstraint> MatConstraint : MaterialConstraints)
+				{
+					GenerateProcess->UpdateSourceBakeMaterialConstraint(MatConstraint.Key, MatConstraint.Value);
+				}
+				for (TPair<UTexture2D*, UGenerateStaticMeshLODProcess::ETextureBakingConstraint> TexConstraint : TextureConstraints)
+				{
+					GenerateProcess->UpdateSourceBakeTextureConstraint(TexConstraint.Key, TexConstraint.Value);
+				}
+
+				if (Progress && Progress->Cancelled())
+				{
+					return;
+				}
+
+				GenerateProcess->ComputeDerivedSourceData(Progress);
+
+				if (Progress && Progress->Cancelled())
+				{
+					return;
+				}
+
+				*ResultMesh = GenerateProcess->GetDerivedLOD0Mesh();
+				ResultTangents = GenerateProcess->GetDerivedLOD0MeshTangents();
+				ResultCollision = GenerateProcess->GetDerivedCollision();
+
+				if (ResultMesh->HasAttributes() && ResultTangents.GetTangents().Num() > 0 && ResultTangents.GetBitangents().Num() > 0)
+				{
+					ResultMesh->Attributes()->SetNumNormalLayers(3);
+					ensure(ResultTangents.CopyToOverlays(*ResultMesh));
+				}
+			};
+
+			GenerateProcess->GraphEvalCriticalSection.Lock();
+			DoCompute(Progress);
+			GenerateProcess->GraphEvalCriticalSection.Unlock();
+		}
+	};
+
+	class FGenerateStaticMeshLODAssetOperatorFactory : public IDynamicMeshOperatorFactory
+	{
+
+	public:
+
+		FGenerateStaticMeshLODAssetOperatorFactory(UGenerateStaticMeshLODAssetTool* AutoLODTool, FTransformSRT3d ResultTransform) :
+			AutoLODTool(AutoLODTool), 
+			ResultTransform(ResultTransform) 
+		{}
+
+		// IDynamicMeshOperatorFactory API
+		virtual TUniquePtr<FDynamicMeshOperator> MakeNewOperator() override
+		{
+			check(AutoLODTool);
+			TUniquePtr<FGenerateStaticMeshLODAssetOperatorOp> Op = MakeUnique<FGenerateStaticMeshLODAssetOperatorOp>();		
+			Op->GenerateProcess = AutoLODTool->GenerateProcess;
+			Op->GeneratorSettings_Preprocess = AutoLODTool->BasicProperties->Preprocessing;
+			Op->GeneratorSettings_MeshGeneration = AutoLODTool->BasicProperties->MeshGeneration;
+			Op->GeneratorSettings_Simplify = AutoLODTool->BasicProperties->Simplification;
+			Op->GeneratorSettings_Normals = AutoLODTool->BasicProperties->Normals;
+			Op->GeneratorSettings_Texture = AutoLODTool->BasicProperties->TextureBaking;
+			Op->GeneratorSettings_UV = AutoLODTool->BasicProperties->UVGeneration;
+			Op->GeneratorSettings_Collision.CollisionGroupLayerName = AutoLODTool->BasicProperties->CollisionGroupLayerName;
+			Op->GeneratorSettings_Collision = AutoLODTool->BasicProperties->SimpleCollision;
+
+			Op->MaterialConstraints.Reset();
+			for (const FGenerateStaticMeshLOD_MaterialConfig& MatConfig : AutoLODTool->TextureProperties->Materials)
+			{
+				Op->MaterialConstraints.Add( TPair<UMaterialInterface*, UGenerateStaticMeshLODProcess::EMaterialBakingConstraint>( 
+					MatConfig.Material, static_cast<UGenerateStaticMeshLODProcess::EMaterialBakingConstraint>(MatConfig.Constraint) ) );
+			}
+			Op->TextureConstraints.Reset();
+			for (const FGenerateStaticMeshLOD_TextureConfig& TexConfig : AutoLODTool->TextureProperties->Textures)
+			{
+				Op->TextureConstraints.Add( TPair<UTexture2D*, UGenerateStaticMeshLODProcess::ETextureBakingConstraint>( 
+					TexConfig.Texture, static_cast<UGenerateStaticMeshLODProcess::ETextureBakingConstraint>(TexConfig.Constraint) ) );
+			}
+
+			Op->SetResultTransform(ResultTransform);
+			return Op;
+		}
+
+		UGenerateStaticMeshLODAssetTool* AutoLODTool = nullptr;
+		FTransformSRT3d ResultTransform;
+	};
+
+	static void DisplayCriticalWarningMessage(const FString& Message)
+	{
+		FNotificationInfo Info(FText::FromString(Message));
+		Info.ExpireDuration = 5.0f;
+		FSlateNotificationManager::Get().AddNotification(Info);
+
+		UE_LOG(LogTemp, Warning, TEXT("%s"), *Message);
+	}
+}
+
+
+/*
+ * ToolBuilder
+ */
+
+
+const FToolTargetTypeRequirements& UGenerateStaticMeshLODAssetToolBuilder::GetTargetRequirements() const
+{
+	static FToolTargetTypeRequirements TypeRequirements({
+		UPrimitiveComponentBackedTarget::StaticClass(),
+		UStaticMeshBackedTarget::StaticClass()
+		});
+	return TypeRequirements;
+}
+
+bool UGenerateStaticMeshLODAssetToolBuilder::CanBuildTool(const FToolBuilderState& SceneState) const
+{
+	// hack to make multi-tool look like single-tool
+	const int NumTargets = SceneState.TargetManager->CountSelectedAndTargetable(SceneState, GetTargetRequirements());
+	if (NumTargets == 1)
+	{
+		// Ensure that the StaticMesh asset is still valid.
+		// [TODO] Remove this check once we properly handle a failed DuplicateAsset in UGenerateStaticMeshLODProcess::WriteDerivedStaticMeshAsset
+		bool bValidTarget = false;
+		SceneState.TargetManager->EnumerateSelectedAndTargetableComponents(SceneState, GetTargetRequirements(),
+			[&bValidTarget](UActorComponent* Component)
+			{
+				const UStaticMeshComponent* StaticMeshComponent = Cast<UStaticMeshComponent>(Component);
+				UStaticMesh* StaticMesh = StaticMeshComponent ? StaticMeshComponent->GetStaticMesh() : nullptr;
+				if (IsValid(StaticMesh))
+				{
+					bValidTarget = ObjectTools::IsObjectBrowsable(StaticMesh);
+				}
+			});
+		return bValidTarget;
+	}
+	return false;
+}
+
+UMultiSelectionMeshEditingTool* UGenerateStaticMeshLODAssetToolBuilder::CreateNewTool(const FToolBuilderState& SceneState) const
+{
+	UGenerateStaticMeshLODAssetTool* NewTool = NewObject<UGenerateStaticMeshLODAssetTool>(SceneState.ToolManager);
+	NewTool->SetUseAssetEditorMode(bUseAssetEditorMode);
+	NewTool->SetRestrictiveMode(bInRestrictiveMode);
+	return NewTool;
+}
+
+
+
+
+void UGenerateStaticMeshLODAssetToolPresetProperties::PostAction(EGenerateLODAssetToolPresetAction Action)
+{
+	if (ParentTool.IsValid())
+	{
+		ParentTool->RequestPresetAction(Action);
+	}
+}
+
+
+/*
+ * Tool
+ */
+
+
+void UGenerateStaticMeshLODAssetTool::SetUseAssetEditorMode(bool bEnable)
+{
+	bIsInAssetEditorMode = bEnable;
+}
+
+
+void UGenerateStaticMeshLODAssetTool::SetRestrictiveMode(bool bEnable)
+{
+	bRestrictiveMode = bEnable;
+}
+
+
+void UGenerateStaticMeshLODAssetTool::Setup()
+{
+	using GenerateStaticMeshLODAssetLocals::FGenerateStaticMeshLODAssetOperatorFactory;
+
+	UInteractiveTool::Setup();
+
+	SetToolDisplayName(LOCTEXT("ToolName", "AutoLOD"));
+	GetToolManager()->DisplayMessage(
+		LOCTEXT("OnStartStaticMeshLODAssetTool", "Create a new LOD asset"),
+		EToolMessageLevel::UserNotification);
+
+	GenerateProcess = NewObject<UGenerateStaticMeshLODProcess>(this);
+
+	UStaticMeshComponent* StaticMeshComponent = CastChecked<UStaticMeshComponent>(UE::ToolTarget::GetTargetComponent(Targets[0]));
+	UStaticMesh* StaticMesh = StaticMeshComponent->GetStaticMesh();
+
+	FComponentMaterialSet ComponentMaterials = UE::ToolTarget::GetMaterialSet(Targets[0], false);
+	FComponentMaterialSet AssetMaterials = UE::ToolTarget::GetMaterialSet(Targets[0], true);
+	if (ComponentMaterials != AssetMaterials)
+	{
+		GetToolManager()->DisplayMessage(
+			LOCTEXT("GenerateStaticMeshLODAssetTool_DifferentMaterials", "Selected Component and StaticMesh Asset have different Material Sets. Asset Materials will be used."),
+			EToolMessageLevel::UserWarning);
+	}
+	
+
+	FProgressCancel Progress;
+	FScopedSlowTask SlowTask(2, LOCTEXT("UGenerateStaticMeshLODAssetTool_Setup", "Initializing AutoLOD Generator ..."));
+	SlowTask.MakeDialog();
+
+	if (StaticMesh)
+	{
+		SlowTask.EnterProgressFrame(1);
+
+		bool bInitializeOK = GenerateProcess->Initialize(StaticMesh, &Progress);		// Must happen on main thread
+
+		if (Progress.Warnings.Num() > 0)
+		{
+			const FProgressCancel::FMessageInfo& Warning = Progress.Warnings[0];
+			GetToolManager()->DisplayMessage(Warning.MessageText, (EToolMessageLevel)Warning.MessageLevel);
+		}
+
+		if (!bInitializeOK)
+		{
+			GenerateProcess = nullptr;
+			GetToolManager()->DisplayMessage(
+				LOCTEXT("GenerateStaticMeshLODAssetTool_ErrorInitializing", "Error initializing tool process: invalid Static Mesh input"),
+				EToolMessageLevel::UserError);
+			return;
+		}
+
+	}
+	else
+	{
+		GetToolManager()->DisplayMessage(
+			LOCTEXT("GenerateStaticMeshLODAssetTool_NoStaticMesh", "Could not find Static Mesh in selected input"),
+			EToolMessageLevel::UserError);
+		return;
+	}
+
+	FString FullPathWithExtension = UEditorAssetLibrary::GetPathNameForLoadedAsset(StaticMesh);
+	OutputProperties = NewObject<UGenerateStaticMeshLODAssetToolOutputProperties>(this);
+	AddToolPropertySource(OutputProperties);
+	OutputProperties->NewAssetName = FPaths::GetBaseFilename(FullPathWithExtension, true);
+	OutputProperties->GeneratedSuffix = TEXT("_AutoLOD");
+	OutputProperties->RestoreProperties(this);
+	OutputProperties->OutputMode = (bIsInAssetEditorMode) ? EGenerateLODAssetOutputMode::UpdateExistingAsset : EGenerateLODAssetOutputMode::CreateNewAsset;
+	OutputProperties->bShowOutputMode = !bIsInAssetEditorMode;
+
+	if (!bRestrictiveMode)
+	{
+		PresetProperties = NewObject<UGenerateStaticMeshLODAssetToolPresetProperties>(this);
+		PresetProperties->RestoreProperties(this);
+		PresetProperties->Initialize(this);
+		AddToolPropertySource(PresetProperties);
+		PresetProperties->WatchProperty(PresetProperties->Preset, [this](TWeakObjectPtr<UStaticMeshLODGenerationSettings>) { OnPresetSelectionChanged(); });
+	}
+
+	SlowTask.EnterProgressFrame(1);
+	BasicProperties = NewObject<UGenerateStaticMeshLODAssetToolProperties>(this);
+	AddToolPropertySource(BasicProperties);
+	BasicProperties->Preprocessing = GenerateProcess->GetCurrentPreprocessSettings();
+	BasicProperties->MeshGeneration = GenerateProcess->GetCurrentSettings();
+	BasicProperties->Simplification = GenerateProcess->GetCurrentSimplifySettings();
+	BasicProperties->Normals = GenerateProcess->GetCurrentNormalsSettings();
+	BasicProperties->TextureBaking = GenerateProcess->GetCurrentTextureSettings();
+	BasicProperties->UVGeneration = GenerateProcess->GetCurrentUVSettings();
+	BasicProperties->SimpleCollision = GenerateProcess->GetCurrentCollisionSettings();
+	// if we defer restore to here, then on first run we get the defaults, but afterwards we get the restored values
+	BasicProperties->RestoreProperties(this);
+
+	// assume this is an Editor-only Tool, so we can rely on this
+	BasicProperties->GetOnModified().AddLambda([this](UObject*, FProperty*) { OnSettingsModified(); });
+
+	// Collision layer name property
+	BasicProperties->WatchProperty(BasicProperties->CollisionGroupLayerName, [this](FName) { OnSettingsModified(); });
+	BasicProperties->InitializeGroupLayers(&(GenerateProcess->GetSourceMesh()));
+
+	TextureProperties = NewObject<UGenerateStaticMeshLODAssetToolTextureProperties>(this);
+	AddToolPropertySource(TextureProperties);
+	for (UTexture2D* BakeTexture : GenerateProcess->GetSourceBakeTextures())
+	{
+		TextureProperties->Textures.Add( FGenerateStaticMeshLOD_TextureConfig{ BakeTexture, EGenerateStaticMeshLOD_BakeConstraint::NoConstraint } );
+	}
+	TextureProperties->WatchProperty(TextureProperties->Textures, [this](TArray<FGenerateStaticMeshLOD_TextureConfig> NewValues) { OnSettingsModified(); });
+	for (UMaterialInterface* BakeMaterial : GenerateProcess->GetSourceBakeMaterials())
+	{
+		TextureProperties->Materials.Add( FGenerateStaticMeshLOD_MaterialConfig{ BakeMaterial, EGenerateStaticMeshLOD_BakeConstraint::NoConstraint } );
+	}
+	TextureProperties->WatchProperty(TextureProperties->Materials, [this](TArray<FGenerateStaticMeshLOD_MaterialConfig> NewValues) { OnSettingsModified(); });
+
+
+	FBoxSphereBounds Bounds = StaticMeshComponent->Bounds;
+	FTransformSRT3d PreviewTransform = UE::ToolTarget::GetLocalToWorldTransform(Targets[0]);
+	PreviewTransform.SetTranslation(PreviewTransform.GetTranslation() + 2.5 * (double)Bounds.BoxExtent.Y * FVector3d::UnitY());
+
+	this->OpFactory = MakeUnique<FGenerateStaticMeshLODAssetOperatorFactory>(this, PreviewTransform);
+	PreviewWithBackgroundCompute = NewObject<UMeshOpPreviewWithBackgroundCompute>(this, "Preview");
+	PreviewWithBackgroundCompute->Setup(GetTargetWorld(), this->OpFactory.Get());
+	ToolSetupUtil::ApplyRenderingConfigurationToPreview(PreviewWithBackgroundCompute->PreviewMesh, nullptr);
+	PreviewWithBackgroundCompute->PreviewMesh->SetTangentsMode(EDynamicMeshComponentTangentsMode::ExternallyProvided);
+
+	PreviewWithBackgroundCompute->OnMeshUpdated.AddLambda([this](UMeshOpPreviewWithBackgroundCompute* Compute) 
+	{
+		if (Compute->HaveEmptyResult())
+		{
+			GetToolManager()->DisplayMessage(LOCTEXT("CannotCreateEmptyMesh", "WARNING: Tool doesn't allow creation of an empty mesh."),
+				EToolMessageLevel::UserWarning);
+		}
+		// Don't clear the message area if we have a non-empty result
+	});
+
+	// For the first computation, display a bounding box with the working material. Otherwise it looks like nothing
+	// is happening. And we don't want to copy over the potentially huge input mesh to be the preview mesh.
+	FGridBoxMeshGenerator MeshGen;
+	MeshGen.Box = UE::Geometry::FOrientedBox3d((FVector3d)Bounds.Origin, (FVector3d)Bounds.BoxExtent);
+	MeshGen.Generate();
+	FDynamicMesh3 BoxMesh(&MeshGen);
+	PreviewWithBackgroundCompute->PreviewMesh->UpdatePreview(MoveTemp(BoxMesh));
+	PreviewWithBackgroundCompute->PreviewMesh->SetTransform(FTransform(FVector(0, 2.5f * Bounds.BoxExtent.Y, 0)));
+
+	PreviewWithBackgroundCompute->OnOpCompleted.AddLambda([this](const FDynamicMeshOperator* Op)
+	{
+		const GenerateStaticMeshLODAssetLocals::FGenerateStaticMeshLODAssetOperatorOp* GenerateLODOp =
+			(const GenerateStaticMeshLODAssetLocals::FGenerateStaticMeshLODAssetOperatorOp*)(Op);
+		check(GenerateLODOp);
+
+		// Clear any previous collision geometry visualization
+		CollisionPreview->RemoveAllLineSets();
+
+		// Must happen on main thread
+		FPhysicsDataCollection PhysicsData;
+		PhysicsData.Geometry = GenerateLODOp->ResultCollision;
+		PhysicsData.CopyGeometryToAggregate();
+		UE::PhysicsTools::InitializeCollisionGeometryVisualization(CollisionPreview, CollisionVizSettings, PhysicsData, 0.f, 16);
+
+		// Must happen on main thread, and GenerateProcess might be in use by an Op somewhere else
+		GenerateProcess->GraphEvalCriticalSection.Lock();
+
+		UGenerateStaticMeshLODProcess::FPreviewMaterials PreviewMaterialSet;
+		GenerateProcess->GetDerivedMaterialsPreview(PreviewMaterialSet);
+		if (PreviewMaterialSet.Materials.Num() > 0)
+		{
+			PreviewTextures = PreviewMaterialSet.Textures;
+			PreviewMaterials = PreviewMaterialSet.Materials;
+			PreviewWithBackgroundCompute->PreviewMesh->SetMaterials(PreviewMaterials);
+			TextureProperties->PreviewTextures = PreviewTextures;
+		}
+
+		GenerateProcess->GraphEvalCriticalSection.Unlock();
+	});
+
+	PreviewWithBackgroundCompute->ConfigureMaterials(
+		ToolSetupUtil::GetDefaultSculptMaterial(GetToolManager()),
+		ToolSetupUtil::GetDefaultWorkingMaterial(GetToolManager())
+	);
+
+	CollisionVizSettings = NewObject<UCollisionGeometryVisualizationProperties>(this);
+	CollisionVizSettings->RestoreProperties(this);
+	AddToolPropertySource(CollisionVizSettings);
+	CollisionVizSettings->Initialize(this);
+
+	CollisionPreview = NewObject<UPreviewGeometry>(this);
+	CollisionPreview->CreateInWorld(GetTargetWorld(), (FTransform)PreviewTransform);
+
+	// Trigger any automatic Preset-changed behavior if we started Tool with one already selected
+	// Note: Currently this does nothing, as we rely on the user to manually read/write the preset
+	OnPresetSelectionChanged();
+
+	// Pop up notifications for any warnings
+	for ( const FProgressCancel::FMessageInfo& Warning : Progress.Warnings )
+	{
+		FNotificationInfo NotificationInfo(Warning.MessageText);
+		NotificationInfo.ExpireDuration = 6.0f;
+		NotificationInfo.Hyperlink = FSimpleDelegate::CreateLambda([]() { FGlobalTabmanager::Get()->TryInvokeTab(FName("OutputLog")); });
+		NotificationInfo.HyperlinkText = LOCTEXT("ShowOutputLogHyperlink", "Show Output Log");
+		FSlateNotificationManager::Get().AddNotification(NotificationInfo)->SetCompletionState(SNotificationItem::CS_Fail);
+	}
+}
+
+
+
+bool UGenerateStaticMeshLODAssetTool::ValidateSettings() const
+{
+	if (!BasicProperties || !GenerateProcess)
+	{
+		return true;
+	}
+
+	const FDynamicMesh3& Mesh = GenerateProcess->GetSourceMesh();
+
+	const FName& GroupName = BasicProperties->Preprocessing.FilterGroupLayer;
+	if (!GroupName.IsNone())
+	{
+		bool bFound = false;
+
+		if (Mesh.HasAttributes())
+		{
+			for (int GroupLayerIndex = 0; GroupLayerIndex < Mesh.Attributes()->NumPolygroupLayers(); ++GroupLayerIndex)
+			{
+				const FDynamicMeshPolygroupAttribute* GroupAttribute = Mesh.Attributes()->GetPolygroupLayer(GroupLayerIndex);
+				if (GroupAttribute->GetName() == GroupName)
+				{
+					bFound = true;
+					break;
+				}
+			}
+		}
+
+		if (!bFound)
+		{
+			const FText Message = FText::Format(LOCTEXT("GroupNotFoundWarning", "Group {0} not found on input mesh"), FText::FromName(GroupName));
+			GetToolManager()->DisplayMessage(Message, EToolMessageLevel::UserWarning);
+			return false;
+		}
+	}
+
+
+	const FName& WeightMapName = BasicProperties->Preprocessing.ThickenWeightMapName;
+	if (!WeightMapName.IsNone())
+	{
+		bool bFound = false;
+
+		if (Mesh.HasAttributes())
+		{
+			for (int WeightMapIndex = 0; WeightMapIndex < Mesh.Attributes()->NumWeightLayers(); ++WeightMapIndex)
+			{
+				const FDynamicMeshWeightAttribute* GroupAttribute = Mesh.Attributes()->GetWeightLayer(WeightMapIndex);
+				if (GroupAttribute->GetName() == WeightMapName)
+				{
+					bFound = true;
+					break;
+				}
+			}
+		}
+
+		if (!bFound)
+		{
+			const FText Message = FText::Format(LOCTEXT("WeightMapNotFoundWarning", "Weight Map {0} not found on input mesh"), FText::FromName(WeightMapName));
+			GetToolManager()->DisplayMessage(Message, EToolMessageLevel::UserWarning);
+			return false;
+		}
+	}
+
+	return true;
+}
+
+void UGenerateStaticMeshLODAssetTool::OnSettingsModified()
+{
+	if (BasicProperties)
+	{
+		// Set the transient property which controls the edit condition for ThickenAmount
+		// Note: The tool will display an error if ThickenWeightMapName exists but isn't valid
+		BasicProperties->Preprocessing.bIsThickenAmountEnabled = (
+			BasicProperties->Preprocessing.ThickenWeightMapName.IsNone() == false);
+	}
+
+	bool bOK = ValidateSettings();
+	if (bOK)
+	{
+		GetToolManager()->DisplayMessage({}, EToolMessageLevel::UserWarning);
+	}
+
+	PreviewWithBackgroundCompute->InvalidateResult();
+}
+
+
+
+void UGenerateStaticMeshLODAssetTool::OnShutdown(EToolShutdownType ShutdownType)
+{
+	if (OutputProperties)
+	{
+		OutputProperties->SaveProperties(this);
+	}
+	if (BasicProperties)
+	{
+		BasicProperties->SaveProperties(this);
+	}
+	if (CollisionVizSettings)
+	{
+		CollisionVizSettings->SaveProperties(this);
+	}
+	if (PresetProperties)
+	{
+		PresetProperties->SaveProperties(this);
+	}
+	if (CollisionPreview)
+	{
+		CollisionPreview->Disconnect();
+		CollisionPreview = nullptr;
+	}
+
+	if (ShutdownType == EToolShutdownType::Accept && GenerateProcess)
+	{
+		if (OutputProperties->OutputMode == EGenerateLODAssetOutputMode::UpdateExistingAsset)
+		{
+			UpdateExistingAsset();
+		}
+		else
+		{
+			CreateNewAsset();
+		}
+	}
+
+	if (PreviewWithBackgroundCompute)
+	{
+		PreviewWithBackgroundCompute->Shutdown();
+		PreviewWithBackgroundCompute = nullptr;
+	}
+
+	if (GenerateProcess)
+	{
+		GenerateProcess = nullptr;
+	}
+}
+
+
+bool UGenerateStaticMeshLODAssetTool::CanAccept() const
+{
+	return (PreviewWithBackgroundCompute && PreviewWithBackgroundCompute->HaveValidNonEmptyResult());
+}
+
+
+void UGenerateStaticMeshLODAssetTool::OnTick(float DeltaTime)
+{
+	if (PreviewWithBackgroundCompute)
+	{
+		PreviewWithBackgroundCompute->Tick(DeltaTime);
+	}
+
+	UE::PhysicsTools::UpdateCollisionGeometryVisualization(CollisionPreview, CollisionVizSettings);
+}
+
+
+
+void UGenerateStaticMeshLODAssetTool::CreateNewAsset()
+{
+	check(PreviewWithBackgroundCompute->HaveValidResult());
+	GenerateProcess->UpdateDerivedPathName(OutputProperties->NewAssetName, OutputProperties->GeneratedSuffix);
+
+	check(GenerateProcess->GraphEvalCriticalSection.TryLock());		// No ops should be running
+	GenerateProcess->WriteDerivedAssetData();
+	GenerateProcess->GraphEvalCriticalSection.Unlock();
+}
+
+
+
+void UGenerateStaticMeshLODAssetTool::UpdateExistingAsset()
+{
+	check(PreviewWithBackgroundCompute->HaveValidResult());
+
+	const UStaticMesh* SourceStaticMesh = GenerateProcess->GetSourceStaticMesh();
+	if (SourceStaticMesh->GetPathName().StartsWith(TEXT("/Engine/")))
+	{
+		GenerateStaticMeshLODAssetLocals::DisplayCriticalWarningMessage(FString::Printf(TEXT("CANNOT MODIFY BUILT-IN ENGINE ASSET %s"), *SourceStaticMesh->GetPathName()));
+		return;
+	}
+
+
+	GenerateProcess->UpdateDerivedPathName(OutputProperties->NewAssetName, OutputProperties->GeneratedSuffix);
+
+	check(GenerateProcess->GraphEvalCriticalSection.TryLock());		// No ops should be running
+
+	// only updated HD source if we have no HD source asset. Otherwise we are overwriting with existing lowpoly LOD0.
+	bool bUpdateHDSource =
+		OutputProperties->bSaveInputAsHiResSource &&
+		(GenerateProcess->GetSourceStaticMesh()->IsHiResMeshDescriptionValid() == false);
+
+	GenerateProcess->UpdateSourceAsset(bUpdateHDSource);
+	GenerateProcess->GraphEvalCriticalSection.Unlock();
+}
+
+
+
+
+void UGenerateStaticMeshLODAssetTool::RequestPresetAction(EGenerateLODAssetToolPresetAction ActionType)
+{
+	if (PresetProperties == nullptr)
+	{
+		return;
+	}
+
+	UStaticMeshLODGenerationSettings* CurrentPreset = PresetProperties->Preset.Get();
+	if (CurrentPreset == nullptr)
+	{
+		GetToolManager()->DisplayMessage(
+			LOCTEXT("GenerateStaticMeshLODAssetTool_NoPresetSelected", "No Preset Asset is currently set in the Preset Settings"), EToolMessageLevel::UserError);
+		return;
+	}
+
+	if (ActionType == EGenerateLODAssetToolPresetAction::ReadFromPreset)
+	{
+		UStaticMeshLODGenerationSettings* ApplyPreset = PresetProperties->Preset.Get();
+		if (ApplyPreset)
+		{
+			BasicProperties->Preprocessing = ApplyPreset->Preprocessing;
+			BasicProperties->MeshGeneration = ApplyPreset->MeshGeneration;
+			BasicProperties->Simplification = ApplyPreset->Simplification;
+			BasicProperties->Normals = ApplyPreset->Normals;
+			BasicProperties->TextureBaking = ApplyPreset->TextureBaking;
+			BasicProperties->UVGeneration = ApplyPreset->UVGeneration;
+			BasicProperties->SimpleCollision = ApplyPreset->SimpleCollision;
+			OnSettingsModified();
+		}
+	}
+	else if (ActionType == EGenerateLODAssetToolPresetAction::WriteToPreset)
+	{
+		GetToolManager()->BeginUndoTransaction(LOCTEXT("WriteToPresetAction", "Write Preset"));
+
+		CurrentPreset->SetFlags(RF_Transactional);
+		CurrentPreset->Modify();
+
+		CurrentPreset->Preprocessing = BasicProperties->Preprocessing;
+		CurrentPreset->MeshGeneration = BasicProperties->MeshGeneration;
+		CurrentPreset->Simplification = BasicProperties->Simplification;
+		CurrentPreset->Normals = BasicProperties->Normals;
+		CurrentPreset->TextureBaking = BasicProperties->TextureBaking;
+		CurrentPreset->UVGeneration = BasicProperties->UVGeneration;
+		CurrentPreset->SimpleCollision = BasicProperties->SimpleCollision;
+
+		CurrentPreset->PostEditChange();
+
+		GetToolManager()->EndUndoTransaction();
+	}
+}
+
+
+void UGenerateStaticMeshLODAssetTool::OnPresetSelectionChanged()
+{
+	// Rely on user to decide when to write/read settings to/from the selected preset
+}
+
+
+
+
+
+#undef LOCTEXT_NAMESPACE
diff --git a/Engine/Plugins/Editor/MeshLODToolset/Source/MeshLODToolset/Private/Tools/LODGenerationSettingsAsset.cpp b/Engine/Plugins/Editor/MeshLODToolset/Source/MeshLODToolset/Private/Tools/LODGenerationSettingsAsset.cpp
new file mode 100644
--- /dev/null
+++ b/Engine/Plugins/Editor/MeshLODToolset/Source/MeshLODToolset/Private/Tools/LODGenerationSettingsAsset.cpp
@@ -0,0 +1,7 @@
+// Copyright Epic Games, Inc. All Rights Reserved.
+
+#include "Tools/LODGenerationSettingsAsset.h"
+
+#include UE_INLINE_GENERATED_CPP_BY_NAME(LODGenerationSettingsAsset)
+
+
diff --git a/Engine/Plugins/Editor/MeshLODToolset/Source/MeshLODToolset/Private/Tools/LODGenerationSettingsFactory.cpp b/Engine/Plugins/Editor/MeshLODToolset/Source/MeshLODToolset/Private/Tools/LODGenerationSettingsFactory.cpp
new file mode 100644
--- /dev/null
+++ b/Engine/Plugins/Editor/MeshLODToolset/Source/MeshLODToolset/Private/Tools/LODGenerationSettingsFactory.cpp
@@ -0,0 +1,44 @@
+// Copyright Epic Games, Inc. All Rights Reserved.
+
+#include "Tools/LODGenerationSettingsFactory.h"
+#include "Tools/LODGenerationSettingsAsset.h"
+#include "AssetTypeCategories.h"
+
+#include UE_INLINE_GENERATED_CPP_BY_NAME(LODGenerationSettingsFactory)
+
+#define LOCTEXT_NAMESPACE "UStaticMeshLODGenerationSettingsFactory"
+
+UStaticMeshLODGenerationSettingsFactory::UStaticMeshLODGenerationSettingsFactory(const FObjectInitializer& ObjectInitializer)
+	: Super(ObjectInitializer)
+{
+	SupportedClass = UStaticMeshLODGenerationSettings::StaticClass();
+	bCreateNew = true;
+	bText = true;
+	bEditAfterNew = true;
+	bEditorImport = false;
+}
+
+UObject* UStaticMeshLODGenerationSettingsFactory::FactoryCreateNew(UClass* Class, UObject* InParent, FName Name, EObjectFlags Flags, UObject* Context, FFeedbackContext* Warn)
+{
+	UStaticMeshLODGenerationSettings* NewSettings = NewObject<UStaticMeshLODGenerationSettings>(InParent, Class, Name, Flags);
+	check(NewSettings);
+	return NewSettings;
+}
+
+
+uint32 UStaticMeshLODGenerationSettingsFactory::GetMenuCategories() const
+{
+	return EAssetTypeCategories::None;
+}
+
+FText UStaticMeshLODGenerationSettingsFactory::GetDisplayName() const
+{
+	return LOCTEXT("MenuEntry", "AutoLOD Settings");
+}
+
+FString UStaticMeshLODGenerationSettingsFactory::GetDefaultNewAssetName() const
+{
+	return TEXT("AutoLODSettings");
+}
+
+#undef LOCTEXT_NAMESPACE
diff --git a/Engine/Plugins/Editor/MeshLODToolset/Source/MeshLODToolset/Private/Tools/LODManagerTool.cpp b/Engine/Plugins/Editor/MeshLODToolset/Source/MeshLODToolset/Private/Tools/LODManagerTool.cpp
new file mode 100644
--- /dev/null
+++ b/Engine/Plugins/Editor/MeshLODToolset/Source/MeshLODToolset/Private/Tools/LODManagerTool.cpp
@@ -0,0 +1,694 @@
+// Copyright Epic Games, Inc. All Rights Reserved.
+
+#include "Tools/LODManagerTool.h"
+#include "Drawing/LineSetComponent.h"
+#include "InteractiveToolManager.h"
+
+#include "AssetUtils/MeshDescriptionUtil.h"
+#include "Drawing/PreviewGeometryActor.h"
+#include "MeshDescriptionToDynamicMesh.h"
+#include "Engine/StaticMeshSourceData.h"
+#include "IndexedMeshToDynamicMesh.h"
+#include "ModelingToolTargetUtil.h"
+
+#include "PreviewMesh.h"
+#include "StaticMeshAttributes.h"
+
+// for lightmap access
+#include "Components/StaticMeshComponent.h"
+
+
+#include "StaticMeshResources.h"
+#include "TargetInterfaces/MaterialProvider.h"
+#include "TargetInterfaces/PrimitiveComponentBackedTarget.h"
+#include "TargetInterfaces/StaticMeshBackedTarget.h"
+#include "ToolContextInterfaces.h"
+#include "ToolTargetManager.h"
+
+#include UE_INLINE_GENERATED_CPP_BY_NAME(LODManagerTool)
+
+using namespace UE::Geometry;
+
+#define LOCTEXT_NAMESPACE "ULODManagerTool"
+
+/*
+ * ToolBuilder
+ */
+
+
+const FToolTargetTypeRequirements& ULODManagerToolBuilder::GetTargetRequirements() const
+{
+	static FToolTargetTypeRequirements TypeRequirements({
+		UMaterialProvider::StaticClass(),
+		UPrimitiveComponentBackedTarget::StaticClass(),
+		UStaticMeshBackedTarget::StaticClass()
+		});
+	return TypeRequirements;
+}
+
+bool ULODManagerToolBuilder::CanBuildTool(const FToolBuilderState& SceneState) const
+{
+	// disable multi-selection for now
+	return SceneState.TargetManager->CountSelectedAndTargetable(SceneState, GetTargetRequirements()) == 1;
+}
+
+UMultiSelectionMeshEditingTool* ULODManagerToolBuilder::CreateNewTool(const FToolBuilderState& SceneState) const
+{
+	return NewObject<ULODManagerTool>(SceneState.ToolManager);
+}
+
+
+
+
+void ULODManagerActionPropertySet::PostAction(ELODManagerToolActions Action)
+{
+	if (ParentTool.IsValid() && Cast<ULODManagerTool>(ParentTool))
+	{
+		Cast<ULODManagerTool>(ParentTool)->RequestAction(Action);
+	}
+}
+
+
+
+
+ULODManagerTool::ULODManagerTool()
+{
+}
+
+
+void ULODManagerTool::Setup()
+{
+	UInteractiveTool::Setup();
+
+	if (! ensure(Targets.Num() == 1) )
+	{
+		return;
+	}
+
+
+	LODPreviewProperties = NewObject<ULODManagerPreviewLODProperties>(this);
+	AddToolPropertySource(LODPreviewProperties);
+	LODPreviewProperties->VisibleLOD = this->DefaultLODName;
+	LODPreviewProperties->bShowingDefaultLOD = true;
+	LODPreviewProperties->WatchProperty(LODPreviewProperties->VisibleLOD, [this](FString NewLOD) { bPreviewLODValid = false; });
+	LODPreviewProperties->WatchProperty(LODPreviewProperties->bShowSeams, [this](bool bNewValue) { bPreviewLODValid = false; });
+
+	LODPreview = NewObject<UPreviewMesh>(this);
+	LODPreview->CreateInWorld(GetTargetWorld(), (FTransform)UE::ToolTarget::GetLocalToWorldTransform(Targets[0]));
+	LODPreview->SetTangentsMode(EDynamicMeshComponentTangentsMode::ExternallyProvided);
+	LODPreview->SetVisible(false);
+
+	LODPreviewLines = NewObject<UPreviewGeometry>(this);
+	LODPreviewLines->CreateInWorld(GetTargetWorld(), (FTransform)UE::ToolTarget::GetLocalToWorldTransform(Targets[0]));
+
+	FComponentMaterialSet MaterialSet = UE::ToolTarget::GetMaterialSet(Targets[0]);
+	LODPreview->SetMaterials(MaterialSet.Materials);
+
+	bLODInfoValid = false;
+
+	if (UStaticMeshComponent* StaticMeshComponent = Cast<UStaticMeshComponent>(UE::ToolTarget::GetTargetComponent(Targets[0])))
+	{
+		if (UStaticMesh* StaticMesh = StaticMeshComponent->GetStaticMesh())
+		{
+			if (StaticMesh->IsHiResMeshDescriptionValid())
+			{
+				HiResSourceModelActions = NewObject<ULODManagerHiResSourceModelActions>(this);
+				HiResSourceModelActions->Initialize(this);
+				AddToolPropertySource(HiResSourceModelActions);
+			}
+		}
+	}
+
+	MaterialActions = NewObject<ULODManagerMaterialActions>(this);
+	MaterialActions->Initialize(this);
+	AddToolPropertySource(MaterialActions);
+
+	LODInfoProperties = NewObject<ULODManagerLODProperties>(this);
+	AddToolPropertySource(LODInfoProperties);
+	bLODInfoValid = false;
+
+	SetToolDisplayName(LOCTEXT("ToolName", "Manage LODs"));
+	GetToolManager()->DisplayMessage(
+		LOCTEXT("OnStartTool", "Inspect and Modify LODs of a StaticMesh Asset"),
+		EToolMessageLevel::UserNotification);
+}
+
+
+
+
+void ULODManagerTool::OnShutdown(EToolShutdownType ShutdownType)
+{
+	if (LODPreview)
+	{
+		LODPreview->Disconnect();
+		LODPreviewLines->Disconnect();
+	}
+	for (int32 ComponentIdx = 0; ComponentIdx < Targets.Num(); ComponentIdx++)
+	{
+		UE::ToolTarget::ShowSourceObject(Targets[ComponentIdx]);
+	}
+}
+
+
+
+void ULODManagerTool::RequestAction(ELODManagerToolActions ActionType)
+{
+	if (PendingAction == ELODManagerToolActions::NoAction)
+	{
+		PendingAction = ActionType;
+	}
+}
+
+
+
+void ULODManagerTool::OnTick(float DeltaTime)
+{
+	switch (PendingAction)
+	{
+	case ELODManagerToolActions::DeleteHiResSourceModel:
+		DeleteHiResSourceModel();
+		break;
+
+	case ELODManagerToolActions::MoveHiResToLOD0:
+		MoveHiResToLOD0();
+		break;
+
+	case ELODManagerToolActions::RemoveUnreferencedMaterials:
+		RemoveUnreferencedMaterials();
+		break;
+	}
+	PendingAction = ELODManagerToolActions::NoAction;
+
+	if (bLODInfoValid == false)
+	{
+		UpdateLODInfo();
+	}
+
+	if (bPreviewLODValid == false)
+	{
+		UpdatePreviewLOD();
+	}
+}
+
+
+
+UStaticMesh* ULODManagerTool::GetSingleStaticMesh()
+{
+	if (Targets.Num() > 1)
+	{
+		return nullptr;
+	}
+	UStaticMeshComponent* StaticMeshComponent = Cast<UStaticMeshComponent>(UE::ToolTarget::GetTargetComponent(Targets[0]));
+	if (StaticMeshComponent == nullptr || StaticMeshComponent->GetStaticMesh() == nullptr)
+	{
+		return nullptr;
+	}
+	UStaticMesh* StaticMesh = StaticMeshComponent->GetStaticMesh();
+	return StaticMesh;
+}
+
+
+
+
+
+
+
+void ULODManagerTool::UpdateLODInfo()
+{
+	UStaticMesh* StaticMesh = GetSingleStaticMesh();
+	if (bLODInfoValid == true || StaticMesh == nullptr)
+	{
+		return;
+	}
+	bLODInfoValid = true;
+
+	LODInfoProperties->bNaniteEnabled = StaticMesh->IsNaniteEnabled();
+	LODInfoProperties->KeepTrianglePercent = StaticMesh->NaniteSettings.KeepPercentTriangles * 100;
+
+	TArray<FStaticMaterial> CurMaterialSet = StaticMesh->GetStaticMaterials();
+	LODInfoProperties->Materials.Reset();
+	for (const FStaticMaterial& Material : StaticMesh->GetStaticMaterials())
+	{
+		LODInfoProperties->Materials.Add(Material);
+	}
+
+
+	// per-LOD info
+	LODInfoProperties->SourceLODs.Reset();
+	int32 NumSourceModels = StaticMesh->GetNumSourceModels();
+	for (int32 si = 0; si < NumSourceModels; ++si)
+	{
+		const FStaticMeshSourceModel& LODSourceModel = StaticMesh->GetSourceModel(si);
+		const FMeshDescription* LODMesh = StaticMesh->GetMeshDescription(si);
+		if (LODMesh != nullptr)		// generated LODs do not have mesh description...is that the only way to tell?
+		{
+			FLODManagerLODInfo LODInfo = { LODMesh->Vertices().Num(), LODMesh->Triangles().Num() };
+			LODInfoProperties->SourceLODs.Add(LODInfo);
+		}
+	}
+
+	LODInfoProperties->HiResSource.Reset();
+	if (StaticMesh->IsHiResMeshDescriptionValid())
+	{
+		const FMeshDescription* HiResMesh = StaticMesh->GetHiResMeshDescription();
+		FLODManagerLODInfo LODInfo = { HiResMesh->Vertices().Num(), HiResMesh->Triangles().Num() };
+		LODInfoProperties->HiResSource.Add(LODInfo);
+	}
+
+	LODInfoProperties->RenderLODs.Reset();
+	if (StaticMesh->HasValidRenderData())
+	{
+		const FStaticMeshRenderData* RenderData = StaticMesh->GetRenderData();
+		for (const FStaticMeshLODResources& LODResource : RenderData->LODResources)
+		{
+			FLODManagerLODInfo LODInfo;
+			LODInfo.VertexCount = LODResource.VertexBuffers.PositionVertexBuffer.GetNumVertices();
+			LODInfo.TriangleCount = LODResource.IndexBuffer.GetNumIndices() / 3;
+			LODInfoProperties->RenderLODs.Add(LODInfo);
+		}
+	}
+
+
+	UpdateLODNames();
+}
+
+
+
+
+void ULODManagerTool::UpdateLODNames()
+{
+	UStaticMesh* StaticMesh = GetSingleStaticMesh();
+	if (StaticMesh == nullptr || LODPreviewProperties == nullptr)
+	{
+		return;
+	}
+
+	TArray<FString>& UILODNames = LODPreviewProperties->LODNamesList;
+	UILODNames.Reset();
+	ActiveLODNames.Reset();
+
+	UILODNames.Add(DefaultLODName);
+	ActiveLODNames.Add(DefaultLODName, FLODName());
+
+	if (StaticMesh->IsHiResMeshDescriptionValid())
+	{
+		FString LODName(TEXT("SourceModel HiRes"));
+		UILODNames.Add(LODName);
+		ActiveLODNames.Add(LODName, { -1, -1, 0 });
+	}
+
+	int32 NumSourceModels = StaticMesh->GetNumSourceModels();
+	for (int32 si = 0; si < NumSourceModels; ++si)
+	{
+		const FStaticMeshSourceModel& LODSourceModel = StaticMesh->GetSourceModel(si);
+		const FMeshDescription* LODMesh = StaticMesh->GetMeshDescription(si);
+		if (LODMesh != nullptr)		// generated LODs do not have mesh description...is that the only way to tell?
+		{
+			FString LODName = FString::Printf(TEXT("SourceModel LOD%d"), si);
+			UILODNames.Add(LODName);
+			ActiveLODNames.Add(LODName, { si, -1, -1 });
+		}
+	}
+
+	if (StaticMesh->HasValidRenderData())
+	{
+		const FStaticMeshRenderData* RenderData = StaticMesh->GetRenderData();
+		int32 Index = 0;
+		for (const FStaticMeshLODResources& LODResource : RenderData->LODResources)
+		{
+			FString LODName = FString::Printf(TEXT("RenderData LOD%d"), Index);
+			UILODNames.Add(LODName);
+			ActiveLODNames.Add(LODName, { -1, Index, -1 });
+			Index++;
+		}
+	}
+
+
+	bPreviewLODValid = false;    // should not always be necessary...
+}
+
+
+
+
+
+void ULODManagerTool::UpdatePreviewLines(FLODMeshInfo& MeshInfo)
+{
+	if (MeshInfo.bInfoCached == false)
+	{
+		MeshInfo.bInfoCached = true;
+		for (int32 eid : MeshInfo.Mesh.EdgeIndicesItr())
+		{
+			if (MeshInfo.Mesh.IsBoundaryEdge(eid))
+			{
+				MeshInfo.BoundaryEdges.Add(eid);
+			}
+		}
+	}
+
+	float LineWidthMultiplier = 1.0f;
+	FColor BoundaryEdgeColor(240, 15, 15);
+	float BoundaryEdgeThickness = LineWidthMultiplier * 4.0f;
+	float BoundaryEdgeDepthBias = 2.0f;
+
+	ULineSetComponent* BoundaryLines = LODPreviewLines->FindLineSet(TEXT("BoundaryLines"));
+	if (BoundaryLines == nullptr)
+	{
+		BoundaryLines = LODPreviewLines->AddLineSet(TEXT("BoundaryLines"));
+	}
+	if (BoundaryLines)
+	{
+		BoundaryLines->Clear();
+		for (int32 eid : MeshInfo.BoundaryEdges)
+		{
+			FVector3d A, B;
+			MeshInfo.Mesh.GetEdgeV(eid, A, B);
+			BoundaryLines->AddLine((FVector)A, (FVector)B, BoundaryEdgeColor, BoundaryEdgeThickness, BoundaryEdgeDepthBias);
+		}
+	}
+}
+
+
+void ULODManagerTool::ClearPreviewLines()
+{
+	ULineSetComponent* BoundaryLines = LODPreviewLines->FindLineSet(TEXT("BoundaryLines"));
+	if (BoundaryLines)
+	{
+		BoundaryLines->Clear();
+	}
+}
+
+
+void ULODManagerTool::UpdatePreviewLOD()
+{
+	UStaticMesh* StaticMesh = GetSingleStaticMesh();
+	if (StaticMesh == nullptr || LODPreviewProperties == nullptr || bPreviewLODValid == true)
+	{
+		return;
+	}
+	bPreviewLODValid = true;
+
+	auto SetShowingDefaultLOD = [this](bool bNewValue)
+	{
+		if (bNewValue != LODPreviewProperties->bShowingDefaultLOD)
+		{
+			LODPreviewProperties->bShowingDefaultLOD = bNewValue;
+			NotifyOfPropertyChangeByTool(LODPreviewProperties);
+		}
+	};
+
+	FString SelectedLOD = LODPreviewProperties->VisibleLOD;
+	const FLODName* FoundName = ActiveLODNames.Find(SelectedLOD);
+	if (SelectedLOD.IsEmpty() || FoundName == nullptr || FoundName->IsDefault() )
+	{
+		// currently just showing the source object when Default is selected
+		LODPreview->SetVisible(false);
+		ClearPreviewLines();
+		LODPreviewLines->SetAllVisible(false);
+		UE::ToolTarget::ShowSourceObject(Targets[0]);
+		SetShowingDefaultLOD(true);
+		return;
+	}
+
+	TUniquePtr<FLODMeshInfo>* Found = LODMeshCache.Find(SelectedLOD);
+	if (Found == nullptr)
+	{
+		CacheLODMesh(SelectedLOD, *FoundName);
+		Found = LODMeshCache.Find(SelectedLOD);
+	}
+
+	SetShowingDefaultLOD(Found == nullptr);
+	
+	if (Found)
+	{
+		UE::ToolTarget::HideSourceObject(Targets[0]);
+
+		LODPreview->ReplaceMesh((*Found)->Mesh);
+		LODPreview->SetVisible(true);
+		UpdatePreviewLines(**Found);
+		LODPreviewLines->SetAllVisible(LODPreviewProperties->bShowSeams);
+	}
+	else
+	{
+		LODPreview->SetVisible(false);
+		ClearPreviewLines();
+		LODPreviewLines->SetAllVisible(false);
+		UE::ToolTarget::ShowSourceObject(Targets[0]);
+	}
+}
+
+
+
+
+bool ULODManagerTool::CacheLODMesh(const FString& Name, FLODName LODName)
+{
+	UStaticMesh* StaticMesh = GetSingleStaticMesh();
+	if (!ensure(StaticMesh)) return false;
+
+	FMeshDescriptionToDynamicMesh Converter;
+
+	if (LODName.OtherIndex == 0)		// HiRes SourceModel
+	{
+		FMeshDescription* HiResMeshDescription = StaticMesh->GetHiResMeshDescription();
+		if (HiResMeshDescription != nullptr)
+		{
+			// TODO: we only need to do this copy if normals or tangents are set to auto-generated...
+			FMeshBuildSettings HiResBuildSettings = StaticMesh->GetHiResSourceModel().BuildSettings;
+			FMeshDescription TmpMeshDescription(*HiResMeshDescription);
+			UE::MeshDescription::InitializeAutoGeneratedAttributes(TmpMeshDescription, &HiResBuildSettings);
+
+			TUniquePtr<FLODMeshInfo> NewMeshInfo = MakeUnique<FLODMeshInfo>();
+			Converter.Convert(&TmpMeshDescription, NewMeshInfo->Mesh, true);
+			LODMeshCache.Add(Name, MoveTemp(NewMeshInfo));
+			return true;
+		}
+
+	}
+	else if (LODName.SourceModelIndex >= 0)
+	{
+		FMeshDescription* LODMeshDescription = StaticMesh->GetMeshDescription(LODName.SourceModelIndex);
+		if (LODMeshDescription != nullptr)
+		{
+			// TODO: we only need to do this copy if normals or tangents are set to auto-generated...
+			FMeshBuildSettings LODBuildSettings = StaticMesh->GetSourceModel(LODName.SourceModelIndex).BuildSettings;
+			FMeshDescription TmpMeshDescription(*LODMeshDescription);
+			UE::MeshDescription::InitializeAutoGeneratedAttributes(TmpMeshDescription, &LODBuildSettings);
+
+			TUniquePtr<FLODMeshInfo> NewMeshInfo = MakeUnique<FLODMeshInfo>();
+			Converter.Convert(&TmpMeshDescription, NewMeshInfo->Mesh, true);
+			LODMeshCache.Add(Name, MoveTemp(NewMeshInfo));
+			return true;
+		}
+	}
+	else if (LODName.RenderDataIndex >= 0)
+	{
+		if (StaticMesh->HasValidRenderData())
+		{
+			const FStaticMeshRenderData* RenderData = StaticMesh->GetRenderData();
+			if (LODName.RenderDataIndex < RenderData->LODResources.Num())
+			{
+				const FStaticMeshLODResources& LODResource = RenderData->LODResources[LODName.RenderDataIndex];
+				TUniquePtr<FLODMeshInfo> NewMeshInfo = MakeUnique<FLODMeshInfo>();
+				UE::Conversion::RenderBuffersToDynamicMesh(LODResource.VertexBuffers, LODResource.IndexBuffer, LODResource.Sections, NewMeshInfo->Mesh);
+				LODMeshCache.Add(Name, MoveTemp(NewMeshInfo));
+				return true;
+			}
+		}
+	}
+
+	return false;
+}
+
+
+
+void ULODManagerTool::DeleteHiResSourceModel()
+{
+	GetToolManager()->BeginUndoTransaction(LOCTEXT("DeleteHiResSourceModel", "Delete HiRes Source"));
+	for (int32 ComponentIdx = 0; ComponentIdx < Targets.Num(); ComponentIdx++)
+	{
+		UStaticMeshComponent* StaticMeshComponent = Cast<UStaticMeshComponent>(UE::ToolTarget::GetTargetComponent(Targets[ComponentIdx]));
+		if (StaticMeshComponent == nullptr || StaticMeshComponent->GetStaticMesh() == nullptr)
+		{
+			continue;
+		}
+		UStaticMesh* StaticMesh = StaticMeshComponent->GetStaticMesh();
+
+		if (StaticMesh->IsHiResMeshDescriptionValid())
+		{
+			StaticMesh->Modify();
+
+			StaticMesh->ModifyHiResMeshDescription();
+			StaticMesh->ClearHiResMeshDescription();
+			StaticMesh->CommitHiResMeshDescription();
+
+			StaticMesh->PostEditChange();
+		}
+	}
+
+	GetToolManager()->EndUndoTransaction();
+
+	LODMeshCache.Reset();
+	bLODInfoValid = false;
+}
+
+
+
+void ULODManagerTool::MoveHiResToLOD0()
+{
+	GetToolManager()->BeginUndoTransaction(LOCTEXT("RestoreLOD0", "Move HiRes to LOD0"));
+	for (int32 ComponentIdx = 0; ComponentIdx < Targets.Num(); ComponentIdx++)
+	{
+		UStaticMeshComponent* StaticMeshComponent = Cast<UStaticMeshComponent>(UE::ToolTarget::GetTargetComponent(Targets[ComponentIdx]));
+		if (StaticMeshComponent == nullptr || StaticMeshComponent->GetStaticMesh() == nullptr)
+		{
+			continue;
+		}
+		UStaticMesh* StaticMesh = StaticMeshComponent->GetStaticMesh();
+
+		if (StaticMesh->IsHiResMeshDescriptionValid())
+		{
+			StaticMesh->Modify();
+
+			StaticMesh->ModifyHiResMeshDescription();
+			FMeshDescription* HiResMeshDescription = StaticMesh->GetHiResMeshDescription();
+
+			StaticMesh->ModifyMeshDescription(0);
+			FMeshDescription* LOD0MeshDescription = StaticMesh->GetMeshDescription(0);
+			*LOD0MeshDescription = *HiResMeshDescription;
+
+			StaticMesh->ClearHiResMeshDescription();
+			StaticMesh->CommitHiResMeshDescription();
+
+			StaticMesh->CommitMeshDescription(0);
+
+			StaticMesh->PostEditChange();
+		}
+	}
+
+	GetToolManager()->EndUndoTransaction();
+
+	LODMeshCache.Reset();
+	bLODInfoValid = false;
+}
+
+
+
+
+void ULODManagerTool::RemoveUnreferencedMaterials()
+{
+	GetToolManager()->BeginUndoTransaction(LOCTEXT("RemoveUnreferencedMaterials", "Remove Unreferenced Materials"));
+	for (int32 ComponentIdx = 0; ComponentIdx < Targets.Num(); ComponentIdx++)
+	{
+		UStaticMeshComponent* StaticMeshComponent = Cast<UStaticMeshComponent>(UE::ToolTarget::GetTargetComponent(Targets[ComponentIdx]));
+		if (StaticMeshComponent == nullptr || StaticMeshComponent->GetStaticMesh() == nullptr)
+		{
+			continue;
+		}
+		UStaticMesh* StaticMesh = StaticMeshComponent->GetStaticMesh();
+
+		TArray<FStaticMaterial> CurMaterialSet = StaticMesh->GetStaticMaterials();
+		int32 NumMaterials = CurMaterialSet.Num();
+		TArray<bool> MatUsedFlags;
+		MatUsedFlags.Init(false, NumMaterials);
+
+		TArray<const FMeshDescription*> Meshes;
+		for (int32 k = 0; k < StaticMesh->GetNumSourceModels(); ++k)
+		{
+			if (StaticMesh->IsSourceModelValid(k))
+			{
+				Meshes.Add(StaticMesh->GetMeshDescription(k));
+			}
+		}
+		if (StaticMesh->IsHiResMeshDescriptionValid())
+		{
+			Meshes.Add(StaticMesh->GetHiResMeshDescription());
+		}
+		int32 NumMeshes = Meshes.Num();
+
+
+		// for each mesh, collect list of material indices, and then set MatUsedFlags
+		TArray<TArray<int32>> MeshMaterialIndices;
+		MeshMaterialIndices.SetNum(NumMeshes);
+		for (int32 mi = 0; mi < NumMeshes; ++mi)
+		{
+			TMap<FPolygonGroupID, int32> PolygonGroupToSectionIndex;
+			TMap<FPolygonGroupID, int32> PolygonGroupToMaterialIndex;
+			FStaticMeshConstAttributes Attributes(*Meshes[mi]);
+			TPolygonGroupAttributesConstRef<FName> PolygonGroupImportedMaterialSlotNames = Attributes.GetPolygonGroupMaterialSlotNames();
+			for (FPolygonGroupID PolygonGroupID : Meshes[mi]->PolygonGroups().GetElementIDs() )
+			{
+				int32& SectionIndex = PolygonGroupToSectionIndex.FindOrAdd(PolygonGroupID);
+				int32 MaterialIndex = StaticMesh->GetMaterialIndexFromImportedMaterialSlotName(PolygonGroupImportedMaterialSlotNames[PolygonGroupID]);
+				if (MaterialIndex == INDEX_NONE)
+				{
+					MaterialIndex = PolygonGroupID.GetValue();
+				}
+				PolygonGroupToMaterialIndex.Add(PolygonGroupID, MaterialIndex);
+			}
+
+			for (const FTriangleID TriangleID : Meshes[mi]->Triangles().GetElementIDs())
+			{
+				const FPolygonGroupID PolygonGroupID = Meshes[mi]->GetTrianglePolygonGroup(TriangleID);
+				const int32 MaterialIndex = PolygonGroupToMaterialIndex[PolygonGroupID];
+				MatUsedFlags[MaterialIndex] = true;
+				MeshMaterialIndices[mi].AddUnique(MaterialIndex);
+			}
+		}
+
+		// construct new material set and ID remap
+		TArray<FStaticMaterial> NewMaterialSet;
+		TArray<int32> MaterialIDMap;
+		TArray<int32> RemovedMaterials;
+		TMap<FPolygonGroupID, FPolygonGroupID> RemapPolygonGroups;
+		MaterialIDMap.Init(-1, NumMaterials);
+		for (int32 k = 0; k < NumMaterials; ++k)
+		{
+			if (MatUsedFlags[k] == true )
+			{
+				MaterialIDMap[k] = NewMaterialSet.Num();
+				RemapPolygonGroups.Add(FPolygonGroupID(k), FPolygonGroupID(MaterialIDMap[k]));
+				NewMaterialSet.Add(CurMaterialSet[k]);
+			}
+			else
+			{
+				RemovedMaterials.Add(k);
+			}
+		}
+		if (NewMaterialSet.Num() == CurMaterialSet.Num())
+		{
+			continue;
+		}
+
+		StaticMesh->Modify();
+
+		// update material set
+		StaticMesh->GetStaticMaterials() = NewMaterialSet;
+
+		// now remap groups in each mesh
+		int32 MeshIndex = 0;
+		for (int32 k = 0; k < StaticMesh->GetNumSourceModels(); ++k)
+		{
+			if (StaticMesh->IsSourceModelValid(k))
+			{
+				StaticMesh->ModifyMeshDescription(k);
+				FMeshDescription* Mesh = StaticMesh->GetMeshDescription(k);
+				Mesh->RemapPolygonGroups(RemapPolygonGroups);
+				StaticMesh->CommitMeshDescription(k);
+			}
+		}
+		if (StaticMesh->IsHiResMeshDescriptionValid())
+		{
+			StaticMesh->ModifyHiResMeshDescription();
+			FMeshDescription* HiResMeshDescription = StaticMesh->GetHiResMeshDescription();
+			HiResMeshDescription->RemapPolygonGroups(RemapPolygonGroups);
+			StaticMesh->CommitHiResMeshDescription();
+		}
+
+		StaticMesh->PostEditChange();
+	}
+
+	GetToolManager()->EndUndoTransaction();
+	bLODInfoValid = false;
+}
+
+#undef LOCTEXT_NAMESPACE
+
diff --git a/Engine/Plugins/Editor/MeshLODToolset/Source/MeshLODToolset/Public/AssetDefinition_StaticMeshLODGenerationSettings.h b/Engine/Plugins/Editor/MeshLODToolset/Source/MeshLODToolset/Public/AssetDefinition_StaticMeshLODGenerationSettings.h
new file mode 100644
--- /dev/null
+++ b/Engine/Plugins/Editor/MeshLODToolset/Source/MeshLODToolset/Public/AssetDefinition_StaticMeshLODGenerationSettings.h
@@ -0,0 +1,28 @@
+// Copyright Epic Games, Inc. All Rights Reserved.
+
+#pragma once
+
+#include "AssetDefinitionDefault.h"
+#include "Tools/LODGenerationSettingsAsset.h"
+#include "AssetDefinition_StaticMeshLODGenerationSettings.generated.h"
+
+/**
+ * Asset Definition for UStaticMeshLODGenerationSettings Assets
+ */
+UCLASS()
+class UAssetDefinition_StaticMeshLODGenerationSettings : public UAssetDefinitionDefault
+{
+	GENERATED_BODY()
+
+private:
+	virtual FText GetAssetDisplayName() const override { return NSLOCTEXT("AssetTypeActions", "FAssetTypeActions_StaticMeshLODGenerationSettings", "AutoLOD Settings"); }
+	virtual FLinearColor GetAssetColor() const override { return FLinearColor(FColor(175, 0, 128)); }
+	virtual TSoftClassPtr<UObject> GetAssetClass() const override { return UStaticMeshLODGenerationSettings::StaticClass(); }
+	virtual TConstArrayView<FAssetCategoryPath> GetAssetCategories() const override
+	{
+		static const auto Categories = { EAssetCategoryPaths::Misc };
+		return Categories;
+	}
+	virtual FAssetSupportResponse CanLocalize(const FAssetData& InAsset) const { return FAssetSupportResponse::NotSupported(); }
+};
+
diff --git a/Engine/Plugins/Editor/MeshLODToolset/Source/MeshLODToolset/Public/Graphs/GenerateMeshLODGraph.h b/Engine/Plugins/Editor/MeshLODToolset/Source/MeshLODToolset/Public/Graphs/GenerateMeshLODGraph.h
new file mode 100644
--- /dev/null
+++ b/Engine/Plugins/Editor/MeshLODToolset/Source/MeshLODToolset/Public/Graphs/GenerateMeshLODGraph.h
@@ -0,0 +1,201 @@
+// Copyright Epic Games, Inc. All Rights Reserved.
+
+#pragma once
+
+#include "GeometryFlowGraph.h"
+#include "DynamicMesh/MeshTangents.h"
+
+#include "MeshProcessingNodes/MeshSolidifyNode.h"
+#include "MeshProcessingNodes/MeshVoxMorphologyNode.h"
+#include "MeshProcessingNodes/MeshMakeCleanGeometryNode.h"
+#include "MeshProcessingNodes/MeshSimplifyNode.h"
+#include "MeshProcessingNodes/MeshNormalsNodes.h"
+#include "MeshProcessingNodes/MeshThickenNode.h"
+#include "MeshProcessingNodes/MeshAutoGenerateUVsNode.h"
+#include "MeshProcessingNodes/GenerateConvexHullMeshNode.h"
+#include "DataTypes/MeshImageBakingData.h"
+#include "PhysicsNodes/GenerateSimpleCollisionNode.h"
+
+namespace UE::Geometry { template <typename PixelType> class TImageBuilder; }
+namespace UE::GeometryFlow { struct FNormalMapImage; }
+namespace UE::GeometryFlow { struct FTextureImage; }
+
+using UE::Geometry::FDynamicMesh3;
+
+struct FMeshLODGraphPreFilterSettings
+{
+	FName FilterGroupLayerName = FName();
+};
+
+class FGenerateMeshLODGraph
+{
+public:
+	/**
+	 * Initialize the LOD generation graph. 
+	 * @param SourceMesh if the SourceMesh is known, some optimizations can be made in the graph, but this is optional
+	 */
+	void BuildGraph(const FDynamicMesh3* SourceMeshHint = nullptr);
+
+	int32 AppendTextureBakeNode(const UE::Geometry::TImageBuilder<FVector4f>& SourceImage, const FString& Identifier);
+
+	void AppendMultiTextureBakeNode(const TMap<int32, UE::GeometryFlow::TSafeSharedPtr<UE::Geometry::TImageBuilder<FVector4f>>>& SourceMaterialImages);
+
+	void SetSourceMesh(const FDynamicMesh3& SourceMesh);
+
+
+	void UpdatePreFilterSettings(const FMeshLODGraphPreFilterSettings& PreFilterSettings);
+	const FMeshLODGraphPreFilterSettings& GetCurrentPreFilterSettings() const { return CurrentPreFilterSettings; }
+
+
+	enum class ECoreMeshGeneratorMode
+	{
+		Solidify = 0,
+		SolidifyAndClose = 1,
+		SimplifyOnly = 2,
+		ConvexHull = 3
+	};
+	void UpdateCoreMeshGeneratorMode(ECoreMeshGeneratorMode NewMode);
+	ECoreMeshGeneratorMode GetCurrentCoreMeshGeneratorMode() const { return CurrentCoreMeshGeneratorMode; }
+
+	void UpdateSolidifySettings(const UE::GeometryFlow::FMeshSolidifySettings& SolidifySettings);
+	const UE::GeometryFlow::FMeshSolidifySettings& GetCurrentSolidifySettings() const { return CurrentSolidifySettings; }
+
+	void UpdateMorphologySettings(const UE::GeometryFlow::FVoxClosureSettings& MorphologySettings);
+	const UE::GeometryFlow::FVoxClosureSettings& GetCurrentMorphologySettings() const { return CurrentMorphologySettings; }
+
+	void UpdateMeshCleaningSettings(const UE::GeometryFlow::FMeshMakeCleanGeometrySettings& CleaningSettings);
+	const UE::GeometryFlow::FMeshMakeCleanGeometrySettings& GetCurrentMeshCleaningSettings() const { return CurrentCleanMeshSettings; }
+
+	void UpdateSimplifySettings(const UE::GeometryFlow::FMeshSimplifySettings& SimplifySettings);
+	const UE::GeometryFlow::FMeshSimplifySettings& GetCurrentSimplifySettings() const { return CurrentSimplifySettings; }
+
+	void UpdateGenerateConvexHullMeshSettings(const UE::GeometryFlow::FGenerateConvexHullMeshSettings& ConvexHullSettings);
+	const UE::GeometryFlow::FGenerateConvexHullMeshSettings& GetCurrentGenerateConvexHullMeshSettings() const { return CurrentGenerateConvexHullMeshSettings; }
+
+	void UpdateNormalsSettings(const UE::GeometryFlow::FMeshNormalsSettings& NormalsSettings);
+	const UE::GeometryFlow::FMeshNormalsSettings& GetCurrentNormalsSettings() const { return CurrentNormalsSettings; }
+
+	void UpdateAutoUVSettings(const UE::GeometryFlow::FMeshAutoGenerateUVsSettings& AutoUVSettings);
+	const UE::GeometryFlow::FMeshAutoGenerateUVsSettings& GetCurrentAutoUVSettings() const { return CurrentAutoUVSettings; }
+
+	void UpdateBakeCacheSettings(const UE::GeometryFlow::FMeshMakeBakingCacheSettings& BakeCacheSettings);
+	const UE::GeometryFlow::FMeshMakeBakingCacheSettings& GetCurrentBakeCacheSettings() const { return CurrentBakeCacheSettings; }
+
+	void UpdateGenerateSimpleCollisionSettings(const UE::GeometryFlow::FGenerateSimpleCollisionSettings& SimpleCollisionSettings);
+	const UE::GeometryFlow::FGenerateSimpleCollisionSettings& GetCurrentGenerateSimpleCollisionSettings() const { return CurrentGenerateSimpleCollisionSettings; }
+
+	void UpdateThickenWeightMap(const TArray<float>& ThickenWeightMap);
+
+	void UpdateCollisionGroupLayerName(const FName& CollisionGroupLayerName);
+
+	void UpdateThickenSettings(const UE::GeometryFlow::FMeshThickenSettings& ThickenSettings);
+	const UE::GeometryFlow::FMeshThickenSettings& GetCurrentThickenSettings() const { return CurrentThickenSettings; }
+
+	void EvaluateResult(
+		FDynamicMesh3& ResultMesh,
+		UE::Geometry::FMeshTangentsd& ResultTangents,
+		UE::Geometry::FSimpleShapeSet3d& ResultCollision,
+		UE::GeometryFlow::FNormalMapImage& NormalMap,
+		TArray<TUniquePtr<UE::GeometryFlow::FTextureImage>>& TextureImages,
+		UE::GeometryFlow::FTextureImage& MultiTextureImage,
+		FProgressCancel* Progress);
+
+protected:
+
+	friend class UGenerateStaticMeshLODProcess;
+
+	TUniquePtr<UE::GeometryFlow::FGraph> Graph;
+
+
+	UE::GeometryFlow::FGraph::FHandle MeshSourceNode;
+
+	UE::GeometryFlow::FGraph::FHandle FilterGroupsLayerNameNode;	// FNameSourceNode that defines name of polygroups layer
+	UE::GeometryFlow::FGraph::FHandle FilterTrianglesNode;
+	FMeshLODGraphPreFilterSettings CurrentPreFilterSettings;
+
+	UE::GeometryFlow::FGraph::FHandle MeshGeneratorSwitchNode;
+	ECoreMeshGeneratorMode CurrentCoreMeshGeneratorMode = ECoreMeshGeneratorMode::SolidifyAndClose;
+
+	UE::GeometryFlow::FGraph::FHandle SolidifyNode;
+	UE::GeometryFlow::FGraph::FHandle SolidifySettingsNode;
+	UE::GeometryFlow::FMeshSolidifySettings CurrentSolidifySettings;
+
+	UE::GeometryFlow::FGraph::FHandle MorphologyNode;
+	UE::GeometryFlow::FGraph::FHandle MorphologySettingsNode;
+	UE::GeometryFlow::FVoxClosureSettings CurrentMorphologySettings;
+
+	UE::GeometryFlow::FGraph::FHandle CleanMeshNode;
+	UE::GeometryFlow::FGraph::FHandle CleanMeshSettingsNode;
+	UE::GeometryFlow::FMeshMakeCleanGeometrySettings CurrentCleanMeshSettings;
+
+	UE::GeometryFlow::FGraph::FHandle SimplifyNode;
+	UE::GeometryFlow::FGraph::FHandle SimplifySettingsNode;
+	UE::GeometryFlow::FMeshSimplifySettings CurrentSimplifySettings;
+
+	UE::GeometryFlow::FGraph::FHandle GenerateConvexHullMeshNode;
+	UE::GeometryFlow::FGraph::FHandle GenerateConvexHullMeshSettingsNode;
+	UE::GeometryFlow::FGenerateConvexHullMeshSettings CurrentGenerateConvexHullMeshSettings;
+
+	UE::GeometryFlow::FGraph::FHandle NormalsNode;
+	UE::GeometryFlow::FGraph::FHandle NormalsSettingsNode;
+	UE::GeometryFlow::FMeshNormalsSettings CurrentNormalsSettings;
+
+	UE::GeometryFlow::FGraph::FHandle AutoUVNode;
+	UE::GeometryFlow::FGraph::FHandle AutoUVSettingsNode;
+	UE::GeometryFlow::FMeshAutoGenerateUVsSettings CurrentAutoUVSettings;
+
+	//UE::GeometryFlow::FGraph::FHandle RecomputeUVNode;
+	//UE::GeometryFlow::FGraph::FHandle RecomputeUVSettingsNode;
+
+	UE::GeometryFlow::FGraph::FHandle RepackUVNode;
+	UE::GeometryFlow::FGraph::FHandle RepackUVSettingsNode;
+
+	UE::GeometryFlow::FGraph::FHandle TangentsNode;
+	UE::GeometryFlow::FGraph::FHandle TangentsSettingsNode;
+
+	UE::GeometryFlow::FGraph::FHandle BakeCacheNode;
+	UE::GeometryFlow::FGraph::FHandle BakeCacheSettingsNode;
+	UE::GeometryFlow::FMeshMakeBakingCacheSettings CurrentBakeCacheSettings;
+
+	UE::GeometryFlow::FGraph::FHandle BakeNormalMapNode;
+	UE::GeometryFlow::FGraph::FHandle BakeNormalMapSettingsNode;
+
+	UE::GeometryFlow::FGraph::FHandle BakeMultiTextureNode;
+	UE::GeometryFlow::FGraph::FHandle BakeMultiTextureSettingsNode;
+	UE::GeometryFlow::FGraph::FHandle MaterialIDTextureSourceNode;
+
+	UE::GeometryFlow::FGraph::FHandle ThickenNode;
+	UE::GeometryFlow::FGraph::FHandle ThickenSettingsNode;
+	UE::GeometryFlow::FGraph::FHandle ThickenWeightMapNode;
+	UE::GeometryFlow::FMeshThickenSettings CurrentThickenSettings;
+
+	UE::GeometryFlow::FGraph::FHandle GroupLayerNameNode;
+
+	struct FBakeTextureGraphInfo
+	{
+		int32 Index;
+		FString Identifier;
+		UE::GeometryFlow::FGraph::FHandle TexSourceNode;
+		UE::GeometryFlow::FGraph::FHandle BakeNode;
+	};
+	TArray<FBakeTextureGraphInfo> BakeTextureNodes;
+
+	UE::GeometryFlow::FGraph::FHandle DecomposeMeshForCollisionNode;
+	
+	UE::GeometryFlow::FGraph::FHandle GenerateSimpleCollisionNode;
+	UE::GeometryFlow::FGraph::FHandle GenerateSimpleCollisionSettingsNode;
+	UE::GeometryFlow::FGenerateSimpleCollisionSettings CurrentGenerateSimpleCollisionSettings;
+
+	UE::GeometryFlow::FGraph::FHandle CollisionOutputNode;
+	UE::GeometryFlow::FGraph::FHandle MeshOutputNode;
+	UE::GeometryFlow::FGraph::FHandle TangentsOutputNode;
+
+	FName CollisionGroupLayerName = TEXT("Default");
+	
+};
+
+#if UE_ENABLE_INCLUDE_ORDER_DEPRECATED_IN_5_2
+#include "DataTypes/NormalMapData.h"
+#include "DataTypes/TextureImageData.h"
+#include "MeshProcessingNodes/MeshDeleteTrianglesNode.h"
+#endif
diff --git a/Engine/Plugins/Editor/MeshLODToolset/Source/MeshLODToolset/Public/Graphs/GenerateStaticMeshLODProcess.h b/Engine/Plugins/Editor/MeshLODToolset/Source/MeshLODToolset/Public/Graphs/GenerateStaticMeshLODProcess.h
new file mode 100644
--- /dev/null
+++ b/Engine/Plugins/Editor/MeshLODToolset/Source/MeshLODToolset/Public/Graphs/GenerateStaticMeshLODProcess.h
@@ -0,0 +1,652 @@
+// Copyright Epic Games, Inc. All Rights Reserved.
+
+#pragma once
+
+#include "DataTypes/NormalMapData.h"
+#include "Engine/StaticMesh.h"
+
+
+
+#include "DataTypes/TextureImageData.h"
+#include "DynamicMesh/MeshTangents.h"
+#include "ShapeApproximation/SimpleShapeSet3.h"
+#include "GenerateStaticMeshLODProcess.generated.h"
+
+class FGenerateMeshLODGraph;
+
+
+class UTexture2D;
+class UMaterialInstanceConstant;
+class UMaterialInstanceDynamic;
+struct FMeshDescription;
+using UE::Geometry::FDynamicMesh3;
+
+
+
+UENUM()
+enum class EGenerateStaticMeshLODProcess_MeshGeneratorModes : uint8
+{
+	//~NOTE: must keep in sync with FGenerateMeshLODGraph::ECoreMeshGeneratorMode
+
+	/**
+	 * Generate a mesh using Marching Cubes to remesh the input shape.
+	 * Use Winding Numbers to determine what is inside
+	 */
+	Solidify = 0,
+
+	/**
+	 * Like Solidify, but Dilate and Contract the shape to delete small
+	 * negative features (sharp inner corners, small holes) before remeshing
+	 */
+	SolidifyAndClose = 1,
+
+	/** Generate a mesh by simplifying input Mesh attributes and filling small holes */
+	CleanAndSimplify = 2,
+
+	/** Generate a mesh by calculating the Convex Hull of the input shape */
+	ConvexHull = 3
+};
+
+
+USTRUCT()
+struct FGenerateStaticMeshLODProcessSettings
+{
+	GENERATED_BODY()
+
+	/** Method used to generate the initial mesh for AutoLOD processing */
+	UPROPERTY(EditAnywhere, Category = MeshGenerator, meta = (DisplayName="Mesh Generator"))
+	EGenerateStaticMeshLODProcess_MeshGeneratorModes MeshGenerator = EGenerateStaticMeshLODProcess_MeshGeneratorModes::SolidifyAndClose;
+
+	// Solidify settings
+
+	/** Target number of voxels along the maximum dimension for Solidify operation */
+	UPROPERTY(EditAnywhere, Category = Solidify, meta = (UIMin = "8", UIMax = "1024", ClampMin = "8", ClampMax = "1024", DisplayName="Voxel Resolution",
+	EditConditionHides, EditCondition = "MeshGenerator == EGenerateStaticMeshLODProcess_MeshGeneratorModes::Solidify || MeshGenerator == EGenerateStaticMeshLODProcess_MeshGeneratorModes::SolidifyAndClose"))
+	int SolidifyVoxelResolution = 128;
+
+	/** Winding number threshold to determine what is considered inside the mesh during Solidify */
+	UPROPERTY(EditAnywhere, Category = Solidify, AdvancedDisplay, meta = (UIMin = "0.1", UIMax = ".9", ClampMin = "-10", ClampMax = "10",
+	EditConditionHides, EditCondition = "MeshGenerator == EGenerateStaticMeshLODProcess_MeshGeneratorModes::Solidify || MeshGenerator == EGenerateStaticMeshLODProcess_MeshGeneratorModes::SolidifyAndClose"))
+	float WindingThreshold = 0.5f;
+
+
+	// Morphology settings
+
+	/** Offset distance in the Morpohological Closure operation */
+	UPROPERTY(EditAnywhere, Category = Morphology, meta = (UIMin = "0", UIMax = "100", ClampMin = "0", ClampMax = "1000", DisplayName = "Closure Distance",
+	EditConditionHides, EditCondition = "MeshGenerator == EGenerateStaticMeshLODProcess_MeshGeneratorModes::SolidifyAndClose"))
+	float ClosureDistance = 1.0f;
+
+
+	// Convex Hull Settings
+
+	/** Whether to subsample input vertices using a regular grid before computing the convex hull */
+	UPROPERTY(EditAnywhere, Category = Collision, meta = (EditConditionHides, EditCondition = "MeshGenerator == EGenerateStaticMeshLODProcess_MeshGeneratorModes::ConvexHull"))
+	bool bPrefilterVertices = true;
+
+	/** Grid resolution (along the maximum-length axis) for subsampling before computing the convex hull */
+	UPROPERTY(EditAnywhere, Category = Collision, meta = (NoSpinbox = "true", EditConditionHides, EditCondition = "MeshGenerator == EGenerateStaticMeshLODProcess_MeshGeneratorModes::ConvexHull && bPrefilterVertices", ClampMin = "4", ClampMax = "100"))
+	int PrefilterGridResolution = 10;
+
+};
+
+
+
+USTRUCT()
+struct FGenerateStaticMeshLODProcess_PreprocessSettings
+{
+	GENERATED_BODY()
+
+	//
+	// NOTE: Customization of widgets for these properties happens in FAutoLODToolDetails::CustomizeDetails(). If you add, remove, or 
+	// change any of the properties in this struct, you must also update that function!
+	//
+
+	// Filter settings
+
+	/** Group layer to use for filtering out detail before processing */
+	UPROPERTY(EditAnywhere, Category = Preprocessing, meta = (DisplayName = "Detail Filter Group Layer"))
+	FName FilterGroupLayer = FName();
+
+	// Thicken settings
+
+	// TODO When ThickenWeightMapName is empty the behavior should be that the given ThickenAmount
+	// is applied to the entire mesh. This is not yet implemented so we disabled ThickenAmount to
+	// reduce confusion.
+
+	/** Weight map used during mesh thickening */
+	UPROPERTY(EditAnywhere, Category = Preprocessing, meta = (DisplayName = "Thicken Weight Map"))
+	FName ThickenWeightMapName = FName();
+
+	/**
+	 * Amount to thicken the mesh prior to Solidifying. The thicken weight map values are multiplied by this value.
+	 * Thickening is primarily intended to repair shape erosion that can occur when using the Solidify Mesh Generators.
+	 */
+	UPROPERTY(EditAnywhere, Category = Preprocessing, meta = (EditCondition = "bIsThickenAmountEnabled", HideEditConditionToggle, UIMin = "0", UIMax = "100", ClampMin = "0", ClampMax = "1000"))
+	float ThickenAmount = 0.0f;
+
+	// Transient property, not exposed to the user.
+	UPROPERTY(meta = (TransientToolProperty))
+	bool bIsThickenAmountEnabled = false;
+};
+
+
+
+
+UENUM()
+enum class EGenerateStaticMeshLODProcess_SimplifyMethod : uint8
+{
+	//~NOTE: must keep in sync with UE::GeometryFlow::EMeshSimplifyTargetType
+	TriangleCount = 0,
+	VertexCount = 1,
+	TrianglePercentage = 2,
+	GeometricTolerance = 3
+};
+
+
+USTRUCT()
+struct FGenerateStaticMeshLODProcess_SimplifySettings
+{
+	GENERATED_BODY()
+
+	/** Method used to simplify the mesh created in the Mesh Generation step */
+	UPROPERTY(EditAnywhere, Category = Simplify, meta = (DisplayName = "Simplification Target"))
+	EGenerateStaticMeshLODProcess_SimplifyMethod Method = EGenerateStaticMeshLODProcess_SimplifyMethod::GeometricTolerance;
+
+	/** Target triangle/vertex count after simplification */
+	UPROPERTY(EditAnywhere, Category = Simplify, meta = (UIMin = "1", ClampMin = "1", EditConditionHides, EditCondition = "Method == EGenerateStaticMeshLODProcess_SimplifyMethod::TriangleCount || Method == EGenerateStaticMeshLODProcess_SimplifyMethod::VertexCount"))
+	int TargetCount = 500;
+
+	/** Target triangle/vertex count after simplification */
+	UPROPERTY(EditAnywhere, Category = Simplify, meta = (UIMin = "0", UIMax = "100", ClampMin = "0", ClampMax = "100", EditConditionHides, EditCondition = "Method == EGenerateStaticMeshLODProcess_SimplifyMethod::TrianglePercentage"))
+	float TargetPercentage = 5.0;
+
+	/** Deviation in World Units (Centimeters) */
+	UPROPERTY(EditAnywhere, Category = Simplify, meta = (UIMin = "0", UIMax = "100", ClampMin = "0", ClampMax = "100", EditConditionHides, EditCondition = "Method == EGenerateStaticMeshLODProcess_SimplifyMethod::GeometricTolerance"))
+	float Tolerance = 0.5;
+};
+
+
+
+
+UENUM()
+enum class EGenerateStaticMeshLODProcess_NormalsMethod : uint8
+{
+	FromAngleThreshold = 0,
+	PerVertex = 1,
+	PerTriangle = 2
+};
+
+
+USTRUCT()
+struct FGenerateStaticMeshLODProcess_NormalsSettings
+{
+	GENERATED_BODY()
+
+	/** Type of Normals to generate */
+	UPROPERTY(EditAnywhere, Category = Normals, meta = (DisplayName = "Normals Method"))
+	EGenerateStaticMeshLODProcess_NormalsMethod Method = EGenerateStaticMeshLODProcess_NormalsMethod::FromAngleThreshold;
+
+	/** Split Normals (ie sharp edge) will be created along mesh edges with opening angles above this threshold */
+	UPROPERTY(EditAnywhere, Category = Normals, meta = (UIMin = "0", UIMax = "180", ClampMin = "0", ClampMax = "180", EditConditionHides, EditCondition = "Method == EGenerateStaticMeshLODProcess_NormalsMethod::FromAngleThreshold"))
+	float Angle = 60.0;
+
+	static int32 MapMethodType(int32 From, bool bProcessToGraph);
+};
+
+
+
+
+UENUM()
+enum class EGenerateStaticMeshLODProcess_AutoUVMethod : uint8
+{
+	PatchBuilder = 0,
+	UVAtlas = 1,
+	XAtlas = 2
+};
+
+
+USTRUCT()
+struct FGenerateStaticMeshLODProcess_UVSettings_PatchBuilder
+{
+	GENERATED_BODY()
+
+	/** This parameter controls alignment of the initial patches to creases in the mesh */
+	UPROPERTY(EditAnywhere, Category = "PatchBuilder", meta = (UIMin = "0.1", UIMax = "2.0", ClampMin = "0.01", ClampMax = "100.0"))
+	float CurvatureAlignment = 1.0f;
+
+	/** Number of smoothing steps to apply; this slightly increases distortion but produces more stable results. */
+	UPROPERTY(EditAnywhere, Category = "PatchBuilder", AdvancedDisplay, meta = (UIMin = "0", UIMax = "25", ClampMin = "0", ClampMax = "1000"))
+	int SmoothingSteps = 5;
+
+	/** Smoothing parameter; larger values result in faster smoothing in each step. */
+	UPROPERTY(EditAnywhere, Category = "PatchBuilder", AdvancedDisplay, meta = (UIMin = "0", UIMax = "1.0", ClampMin = "0", ClampMax = "1.0"))
+	float SmoothingAlpha = 0.25f;
+};
+
+USTRUCT()
+struct FGenerateStaticMeshLODProcess_UVSettings
+{
+	GENERATED_BODY()
+
+	UPROPERTY(EditAnywhere, Category = UVGeneration, meta = (DisplayName = "UV Method"))
+	EGenerateStaticMeshLODProcess_AutoUVMethod UVMethod = EGenerateStaticMeshLODProcess_AutoUVMethod::PatchBuilder;
+
+	/** Maximum number of charts to create in UVAtlas */
+	UPROPERTY(EditAnywhere, Category = "UVAtlas", meta = (DisplayName = "Max Charts", UIMin = 0, UIMax = 1000, EditConditionHides, EditCondition = "UVMethod == EGenerateStaticMeshLODProcess_AutoUVMethod::UVAtlas"))
+	int NumUVAtlasCharts = 0;
+
+	/** Number of initial patches mesh will be split into before computing island merging */
+	UPROPERTY(EditAnywhere, Category = "PatchBuilder", meta = (UIMin = "1", UIMax = "1000", ClampMin = "1", ClampMax = "99999999", EditConditionHides, EditCondition = "UVMethod == EGenerateStaticMeshLODProcess_AutoUVMethod::PatchBuilder"))
+	int NumInitialPatches = 100;
+
+	/** Distortion/Stretching Threshold for island merging - larger values increase the allowable UV stretching */
+	UPROPERTY(EditAnywhere, Category = "PatchBuilder", meta = (UIMin = "1.0", UIMax = "5.0", ClampMin = "1.0", EditConditionHides, EditCondition = "UVMethod == EGenerateStaticMeshLODProcess_AutoUVMethod::PatchBuilder"))
+	float MergingThreshold = 1.5f;
+
+	/** UV islands will not be merged if their average face normals deviate by larger than this amount */
+	UPROPERTY(EditAnywhere, Category = "PatchBuilder", meta = (UIMin = "0.0", UIMax = "90.0", ClampMin = "0.0", ClampMax = "180.0", EditConditionHides, EditCondition = "UVMethod == EGenerateStaticMeshLODProcess_AutoUVMethod::PatchBuilder"))
+	float MaxAngleDeviation = 45.0f;
+
+	UPROPERTY(EditAnywhere, Category = "PatchBuilder", meta = (DisplayName = "UV - PatchBuilder", UIMin = "0.0", UIMax = "90.0", ClampMin = "0.0", ClampMax = "180.0", EditConditionHides, EditCondition = "UVMethod == EGenerateStaticMeshLODProcess_AutoUVMethod::PatchBuilder"))
+	FGenerateStaticMeshLODProcess_UVSettings_PatchBuilder PatchBuilder;
+};
+
+
+
+
+UENUM()
+enum class EGenerateStaticMeshLODBakeResolution
+{
+	Resolution16 = 16 UMETA(DisplayName = "16 x 16"),
+	Resolution32 = 32 UMETA(DisplayName = "32 x 32"),
+	Resolution64 = 64 UMETA(DisplayName = "64 x 64"),
+	Resolution128 = 128 UMETA(DisplayName = "128 x 128"),
+	Resolution256 = 256 UMETA(DisplayName = "256 x 256"),
+	Resolution512 = 512 UMETA(DisplayName = "512 x 512"),
+	Resolution1024 = 1024 UMETA(DisplayName = "1024 x 1024"),
+	Resolution2048 = 2048 UMETA(DisplayName = "2048 x 2048"),
+	Resolution4096 = 4096 UMETA(DisplayName = "4096 x 4096"),
+	Resolution8192 = 8192 UMETA(DisplayName = "8192 x 8192")
+};
+
+
+
+USTRUCT()
+struct FGenerateStaticMeshLODProcess_TextureSettings
+{
+	GENERATED_BODY()
+
+	/** Resolution for normal map and texture baking */
+	UPROPERTY(EditAnywhere, Category = Baking , meta = (DisplayName = "Bake Image Res"))
+	EGenerateStaticMeshLODBakeResolution BakeResolution = EGenerateStaticMeshLODBakeResolution::Resolution1024;
+
+	/** How far away from the output mesh to search for input mesh during baking */
+	UPROPERTY(EditAnywhere, Category = Baking, meta = (UIMin = "0", UIMax = "100", ClampMin = "0", ClampMax = "1000", DisplayName = "Bake Thickness"))
+	float BakeThickness = 5.0f;
+
+	/** Generate an atlassed texture where multiple textures are combined into a single output texture */
+	UPROPERTY(EditAnywhere, Category = Baking, meta = (DisplayName = "Combine Textures"))
+	bool bCombineTextures = true;
+};
+
+
+
+
+UENUM()
+enum class EGenerateStaticMeshLODSimpleCollisionGeometryType : uint8
+{
+	//~NOTE: This must be kept in sync with ESimpleCollisionGeometryType in GenerateSimpleCollisionNode.h
+	AlignedBoxes,
+	OrientedBoxes,
+	MinimalSpheres,
+	Capsules,
+	ConvexHulls,
+	SweptHulls,
+	MinVolume,
+	None
+};
+
+
+UENUM()
+enum class EGenerateStaticMeshLODProjectedHullAxisMode : uint8
+{
+	//~NOTE: This must be kept in sync with FMeshSimpleShapeApproximation::EProjectedHullAxisMode in MeshSimpleShapeApproximation.h
+	X = 0,
+	Y = 1,
+	Z = 2,
+	SmallestBoxDimension = 3,
+	SmallestVolume = 4
+};
+
+
+
+USTRUCT()
+struct FGenerateStaticMeshLODProcess_CollisionSettings
+{
+	GENERATED_BODY()
+
+	// Transient property, not set directly by the user. The user controls a CollisionGroupLayerName dropdown property
+	// on the Tool and that value is copied here.
+	UPROPERTY(meta = (TransientToolProperty))
+	FName CollisionGroupLayerName = TEXT("Default");
+
+	/** Type of simple collision objects to produce */
+	UPROPERTY(EditAnywhere, Category = Collision, meta = (DisplayName = "Collision Type"))
+	EGenerateStaticMeshLODSimpleCollisionGeometryType CollisionType = EGenerateStaticMeshLODSimpleCollisionGeometryType::ConvexHulls;
+
+	// Convex Hull Settings
+
+	/** Target triangle count for each convex hull after simplification */
+	UPROPERTY(EditAnywhere, Category = Collision, meta = (NoSpinbox = "true", DisplayName = "Convex Tri Count",
+														  EditConditionHides, EditCondition = "CollisionType == EGenerateStaticMeshLODSimpleCollisionGeometryType::ConvexHulls"))
+	int ConvexTriangleCount = 50;
+
+	/** Whether to subsample input vertices using a regular grid before computing the convex hull */
+	UPROPERTY(EditAnywhere, Category = Collision, meta = (EditConditionHides, EditCondition = "CollisionType == EGenerateStaticMeshLODSimpleCollisionGeometryType::ConvexHulls"))
+	bool bPrefilterVertices = true;
+
+	/** Grid resolution (along the maximum-length axis) for subsampling before computing the convex hull */
+	UPROPERTY(EditAnywhere, Category = Collision, meta = (NoSpinbox = "true", EditConditionHides, EditCondition = "CollisionType == EGenerateStaticMeshLODSimpleCollisionGeometryType::ConvexHulls && bPrefilterVertices", UIMin = 4, UIMax = 100))
+	int PrefilterGridResolution = 10;
+
+
+	// Swept Convex Hull Settings
+
+	/** Whether to simplify polygons used for swept convex hulls */
+	UPROPERTY(EditAnywhere, Category = Collision, meta = (EditConditionHides, EditCondition = "CollisionType == EGenerateStaticMeshLODSimpleCollisionGeometryType::SweptHulls"))
+	bool bSimplifyPolygons = true;
+
+	/** Target minumum edge length for simplified swept convex hulls */
+	UPROPERTY(EditAnywhere, Category = Collision, meta = (NoSpinbox = "true", UIMin = "0", UIMax = "10", ClampMin = "0", ClampMax = "100000",
+														  EditConditionHides, EditCondition = "CollisionType == EGenerateStaticMeshLODSimpleCollisionGeometryType::SweptHulls"))
+	float HullTolerance = 0.1;
+
+	/** Which axis to sweep along when computing swept convex hulls */
+	UPROPERTY(EditAnywhere, Category = Collision, meta = (EditConditionHides, EditCondition = "CollisionType == EGenerateStaticMeshLODSimpleCollisionGeometryType::SweptHulls"))
+	EGenerateStaticMeshLODProjectedHullAxisMode SweepAxis = EGenerateStaticMeshLODProjectedHullAxisMode::SmallestVolume;
+};
+
+
+
+UCLASS(Transient)
+class UGenerateStaticMeshLODProcess : public UObject
+{
+	GENERATED_BODY()
+public:
+
+	/**
+	 * Initialize the Process, this will:
+	 *  - read the Source Mesh, either LOD0 or the HiRes LOD if available
+	 *  - Extract all Materials and Textures from the Source Asset and determine what needs to be baked
+	 *  - Read all the Source Texture Images that need to be baked
+	 *  - call InitializeGenerator() to set up the compute graph
+	 */
+	[[nodiscard]] bool Initialize(UStaticMesh* SourceMesh, FProgressCancel* Progress = nullptr);
+
+
+	/** @return the source UStaticMesh passed to Initialize() */
+	UStaticMesh* GetSourceStaticMesh() const { return SourceStaticMesh; }
+	/** @return the converted source MeshDescription as a FDynamicMesh3, valid after Initialize() finishes */
+	const FDynamicMesh3& GetSourceMesh() const { return SourceMesh; }
+
+
+
+	//
+	// Settings Objects for various stages of the Generator Process
+	//
+
+	const FGenerateStaticMeshLODProcess_PreprocessSettings& GetCurrentPreprocessSettings() const { return CurrentSettings_Preprocess; }
+	void UpdatePreprocessSettings(const FGenerateStaticMeshLODProcess_PreprocessSettings& NewSettings);
+
+	const FGenerateStaticMeshLODProcessSettings& GetCurrentSettings() const { return CurrentSettings; }
+	void UpdateSettings(const FGenerateStaticMeshLODProcessSettings& NewSettings);
+
+	const FGenerateStaticMeshLODProcess_SimplifySettings& GetCurrentSimplifySettings() const { return CurrentSettings_Simplify; }
+	void UpdateSimplifySettings(const FGenerateStaticMeshLODProcess_SimplifySettings& NewSettings);
+
+	const FGenerateStaticMeshLODProcess_NormalsSettings& GetCurrentNormalsSettings() const { return CurrentSettings_Normals; }
+	void UpdateNormalsSettings(const FGenerateStaticMeshLODProcess_NormalsSettings& NewSettings);
+
+	const FGenerateStaticMeshLODProcess_TextureSettings& GetCurrentTextureSettings() const { return CurrentSettings_Texture; }
+	void UpdateTextureSettings(const FGenerateStaticMeshLODProcess_TextureSettings& NewSettings);
+
+	const FGenerateStaticMeshLODProcess_UVSettings& GetCurrentUVSettings() const { return CurrentSettings_UV; }
+	void UpdateUVSettings(const FGenerateStaticMeshLODProcess_UVSettings& NewSettings);
+
+	const FGenerateStaticMeshLODProcess_CollisionSettings& GetCurrentCollisionSettings() const { return CurrentSettings_Collision; }
+	void UpdateCollisionSettings(const FGenerateStaticMeshLODProcess_CollisionSettings& NewSettings);
+
+
+	//
+	// Material Queries and Constraints
+	//
+
+	/** @return the set of Materials that were identified on the Source Asset */
+	TArray<UMaterialInterface*> GetSourceBakeMaterials() const;
+
+	enum class EMaterialBakingConstraint
+	{
+		NoConstraint = 0,
+		UseExistingMaterial = 1
+	};
+
+	/** Configure the constraint on the specified Source Material */
+	void UpdateSourceBakeMaterialConstraint(UMaterialInterface* Material, EMaterialBakingConstraint Constraint);
+
+
+	//
+	// Texture Queries and Constraints
+	//
+
+	/** @return the set of Textures that were identified as being parameters to the Source Materials, that are identified as potentially needing to be baked to the Derived Asset UVs */
+	TArray<UTexture2D*> GetSourceBakeTextures() const;
+
+
+	enum class ETextureBakingConstraint
+	{
+		NoConstraint = 0,
+		UseExistingTexture = 1
+	};
+
+	/** Configure the constraint on the specified Source Texture */
+	void UpdateSourceBakeTextureConstraint(UTexture2D* Texture, ETextureBakingConstraint Constraint);
+
+
+	//
+	// Output Path / Naming / etc Configuration
+	//
+
+	/**
+	 * Configure the output Name and Suffix for the Mesh Asset. The various strings below will be
+	 * computed based on this input
+	 */
+	void UpdateDerivedPathName(const FString& NewAssetBaseName, const FString& NewAssetSuffix);
+
+
+	const FString& GetSourceAssetPath() const { return SourceAssetPath; }
+	const FString& GetSourceAssetFolder() const { return SourceAssetFolder; }
+	const FString& GetSourceAssetName() const { return SourceAssetName; }
+	static FString GetDefaultDerivedAssetSuffix() { return TEXT("_AutoLOD"); }
+	const FString& GetDerivedAssetName() const { return DerivedAssetName; }
+
+
+	/**
+	 * Given the above configuration, run the Generator to compute the derived Meshes and Textures
+	 */
+
+	bool ComputeDerivedSourceData(FProgressCancel* Progress);
+
+	/** The Generated Mesh */
+	const FDynamicMesh3& GetDerivedLOD0Mesh() const { return DerivedLODMesh; }
+	/** Tangents for the Generated Mesh */
+	const UE::Geometry::FMeshTangentsd& GetDerivedLOD0MeshTangents() const { return DerivedLODMeshTangents; }
+	/** Collision geometry for the Generated Mesh */
+	const UE::Geometry::FSimpleShapeSet3d& GetDerivedCollision() const { return DerivedCollision; }
+
+	struct FPreviewMaterials
+	{
+		TArray<UMaterialInterface*> Materials;
+		TArray<UTexture2D*> Textures;
+	};
+	/** Set of Materials and Textures to be applied to the Generated Mesh */
+	void GetDerivedMaterialsPreview(FPreviewMaterials& MaterialSetOut);
+
+
+	//
+	// Final Output Options
+	//
+
+	/**
+	 * Creates new UStaticMesh Asset, Material Instances, and Textures
+	 */
+	virtual bool WriteDerivedAssetData();
+
+
+	/**
+	 * Update the input UStaticMesh Asset (ie the asset passed to Initialize)
+	 * @param bSetNewHDSourceAsset If true, the original source mesh copied to the UStaticMesh's HiRes SourceModel slot
+	 */
+	virtual void UpdateSourceAsset(bool bSetNewHDSourceAsset = false);
+
+
+
+public:
+	/**
+	 * This FCriticalSection is used to co-ordinate evaluation of the internal compute graph
+	 * by external classes. Generally while doing things like changing settings, reading output, etc,
+	 * the internal compute graph cannot be executing. Often this needs to be locked across multiple
+	 * function calls, so it is exposed publicly.  (TODO: clean this up...)
+	 */
+	FCriticalSection GraphEvalCriticalSection;
+
+protected:
+
+	UPROPERTY()
+	TObjectPtr<UStaticMesh> SourceStaticMesh;
+
+	FString SourceAssetPath;
+	FString SourceAssetFolder;
+	FString SourceAssetName;
+
+	// if true, we are building new LOD0 from the StaticMesh HiRes SourceModel, instead of from the mesh in LOD0
+	bool bUsingHiResSource = false;
+
+	// copy of input MeshDescription with autogenerated attributes computed
+	TSharedPtr<FMeshDescription> SourceMeshDescription;
+	// SourceMeshDescription converted to FDynamicMesh3
+	FDynamicMesh3 SourceMesh;
+
+	struct FTextureInfo
+	{
+		UTexture2D* Texture = nullptr;
+		FName ParameterName;
+		UE::Geometry::FImageDimensions Dimensions;
+		UE::Geometry::TImageBuilder<FVector4f> Image;
+		bool bIsNormalMap = false;
+		bool bIsDefaultTexture = false;
+		bool bShouldBakeTexture = false;
+		bool bIsUsedInMultiTextureBaking = false; 
+		ETextureBakingConstraint Constraint = ETextureBakingConstraint::NoConstraint;
+	};
+
+	int SelectTextureToBake(const TArray<FTextureInfo>& TextureInfos) const;
+
+	// Information about one of the input StaticMesh Materials. Computed in Initialize() and not modified afterwards
+	struct FSourceMaterialInfo
+	{
+		FStaticMaterial SourceMaterial;
+		TArray<FTextureInfo> SourceTextures;
+
+		bool bHasNormalMap = false;						// if true, Material has an exposed NormalMap input texture parameter
+		bool bHasTexturesToBake = false;				// if true, Material has at least one SourceTexture that should be baked
+		bool bIsReusable = false;						// if true, Material doesn't need any texture baking and can be re-used by LOD0
+		bool bIsPreviouslyGeneratedMaterial = false;	// if true, this Material was previously generated by AutoLOD and should be discarded. 
+														// Currently inferred from material being in LOD0 but not HiRes Source
+
+		EMaterialBakingConstraint Constraint = EMaterialBakingConstraint::NoConstraint;
+	};
+
+	// list of initial Source Materials, length equivalent to StaticMesh.StaticMaterials
+	TArray<FSourceMaterialInfo> SourceMaterials;
+
+	FString DerivedSuffix;
+	FString DerivedAssetPath;
+	FString DerivedAssetFolder;
+	FString DerivedAssetName;
+	FString DerivedAssetNameNoSuffix;
+	
+
+	FDynamicMesh3 DerivedLODMesh;				// the new generated LOD0 mesh
+	UE::Geometry::FMeshTangentsd DerivedLODMeshTangents;		// Tangents for DerivedLODMesh
+	UE::Geometry::FSimpleShapeSet3d DerivedCollision;			// Simple Collision for DerivedLODMesh
+
+	// Texture set potentially required by output Material set
+	UE::GeometryFlow::FNormalMapImage DerivedNormalMapImage;	// Normal Map
+	TArray<TUniquePtr<UE::GeometryFlow::FTextureImage>> DerivedTextureImages;	// generated Textures
+	UE::GeometryFlow::FTextureImage DerivedMultiTextureBakeImage;
+	TMap<UTexture2D*, int32> SourceTextureToDerivedTexIndex;	// mapping from input Textures to DerivedTextureImages index
+
+	// Information about an output Material
+	struct FDerivedMaterialInfo
+	{
+		int32 SourceMaterialIndex = -1;				// Index into SourceMaterials
+		bool bUseSourceMaterialDirectly = false;	// If true, do not create/use a Derived Material, directly re-use SourceMaterials[SourceMaterialIndex]
+
+		FStaticMaterial DerivedMaterial;			// points to generated Material
+		TArray<FTextureInfo> DerivedTextures;		// List of generated Textures
+	};
+
+	// list of generated/Derived Materials. Length is the same as SourceMaterials and indices correspond, however some Derived Materials may not be initialized (if Source is reusable or generated)
+	TArray<FDerivedMaterialInfo> DerivedMaterials;
+
+	// This list is for accumulating derived UTexture2D's created during WriteDerivedTextures(). We have to
+	// maintain uproperty references to these or they may be garbage collected
+	UPROPERTY()
+	TSet<TObjectPtr<UTexture2D>> AllDerivedTextures;
+
+	// Derived Normal Map
+	UPROPERTY()
+	TObjectPtr<UTexture2D> DerivedNormalMapTex;
+
+	// For each material participating in multi-texture baking, the parameter name of the texture
+	TMap<int32, FName> MultiTextureParameterName;
+
+	UPROPERTY()
+	TObjectPtr<UTexture2D> DerivedMultiTextureBakeResult;
+
+
+	TUniquePtr<FGenerateMeshLODGraph> Generator;			// active LODGenerator Graph
+	bool InitializeGenerator();
+
+	FGenerateStaticMeshLODProcess_PreprocessSettings CurrentSettings_Preprocess;
+	FGenerateStaticMeshLODProcessSettings CurrentSettings;
+	FGenerateStaticMeshLODProcess_SimplifySettings CurrentSettings_Simplify;
+	FGenerateStaticMeshLODProcess_NormalsSettings CurrentSettings_Normals;
+	FGenerateStaticMeshLODProcess_TextureSettings CurrentSettings_Texture;
+	FGenerateStaticMeshLODProcess_UVSettings CurrentSettings_UV;
+	FGenerateStaticMeshLODProcess_CollisionSettings CurrentSettings_Collision;
+	
+	// This value will be stored in Asset Metadata for each generated Asset Texture, Material/MIC, and StaticMesh,
+	// under the key StaticMeshLOD.GenerationGUID.
+	FString DerivedAssetGUIDKey;
+
+	bool WriteDerivedTexture(UTexture2D* SourceTexture, UTexture2D* DerivedTexture, bool bCreatingNewStaticMeshAsset);
+	bool WriteDerivedTexture(UTexture2D* DerivedTexture, FString BaseTexName, bool bCreatingNewStaticMeshAsset);
+	void WriteDerivedTextures(bool bCreatingNewStaticMeshAsset);
+	void WriteDerivedMaterials(bool bCreatingNewStaticMeshAsset);
+	void UpdateMaterialTextureParameters(UMaterialInstanceConstant* Material, FDerivedMaterialInfo& DerivedMaterialInfo);
+	void WriteDerivedStaticMeshAsset();
+
+	void UpdateSourceStaticMeshAsset(bool bSetNewHDSourceAsset);
+
+	void UpdateMaterialTextureParameters(UMaterialInstanceDynamic* Material, const FSourceMaterialInfo& SourceMaterialInfo,
+		const TMap<UTexture2D*,UTexture2D*>& PreviewTextures, UTexture2D* PreviewNormalMap);
+
+	// Return true if the given path corresponds to a material or texture in SourceMaterials
+	bool IsSourceAsset(const FString& AssetPath) const;
+};
+
+#if UE_ENABLE_INCLUDE_ORDER_DEPRECATED_IN_5_2
+#include "Graphs/GenerateMeshLODGraph.h"
+#endif
diff --git a/Engine/Plugins/Editor/MeshLODToolset/Source/MeshLODToolset/Public/MeshLODToolsetModule.h b/Engine/Plugins/Editor/MeshLODToolset/Source/MeshLODToolset/Public/MeshLODToolsetModule.h
new file mode 100644
--- /dev/null
+++ b/Engine/Plugins/Editor/MeshLODToolset/Source/MeshLODToolset/Public/MeshLODToolsetModule.h
@@ -0,0 +1,34 @@
+// Copyright Epic Games, Inc. All Rights Reserved.
+
+#pragma once
+
+#include "Modules/ModuleInterface.h"
+#include "Templates/SharedPointer.h"
+
+DECLARE_LOG_CATEGORY_EXTERN(LogMeshLODToolset, Log, All);
+
+class IAssetTypeActions;
+
+class FMeshLODToolsetModule : public IModuleInterface
+{
+public:
+
+	/** IModuleInterface implementation */
+	virtual void StartupModule() override;
+	virtual void ShutdownModule() override;
+
+	void OnPostEngineInit();
+
+	// registered asset actions
+	TArray<TSharedRef<IAssetTypeActions>> RegisteredAssetTypeActions;
+
+private:
+	/** StaticClass is not safe on shutdown, so we cache the name, and use this to unregister on shut down */
+	TArray<FName> ClassesToUnregisterOnShutdown;
+
+};
+
+#if UE_ENABLE_INCLUDE_ORDER_DEPRECATED_IN_5_2
+#include "CoreMinimal.h"
+#include "Modules/ModuleManager.h"
+#endif
diff --git a/Engine/Plugins/Editor/MeshLODToolset/Source/MeshLODToolset/Public/Tools/DetailsCustomizations/AutoLODToolCustomizations.h b/Engine/Plugins/Editor/MeshLODToolset/Source/MeshLODToolset/Public/Tools/DetailsCustomizations/AutoLODToolCustomizations.h
new file mode 100644
--- /dev/null
+++ b/Engine/Plugins/Editor/MeshLODToolset/Source/MeshLODToolset/Public/Tools/DetailsCustomizations/AutoLODToolCustomizations.h
@@ -0,0 +1,17 @@
+// Copyright Epic Games, Inc. All Rights Reserved.
+
+#pragma once
+
+#include "IDetailCustomization.h"
+
+/// 
+/// Detail customization for UGenerateStaticMeshLODAssetTool (AKA AutoLOD) properties
+///
+ 
+class FAutoLODToolDetails : public IDetailCustomization
+{
+public:
+	static TSharedRef<IDetailCustomization> MakeInstance();
+	virtual void CustomizeDetails(IDetailLayoutBuilder& DetailBuilder) override;
+};
+
diff --git a/Engine/Plugins/Editor/MeshLODToolset/Source/MeshLODToolset/Public/Tools/GenerateStaticMeshLODAssetTool.h b/Engine/Plugins/Editor/MeshLODToolset/Source/MeshLODToolset/Public/Tools/GenerateStaticMeshLODAssetTool.h
new file mode 100644
--- /dev/null
+++ b/Engine/Plugins/Editor/MeshLODToolset/Source/MeshLODToolset/Public/Tools/GenerateStaticMeshLODAssetTool.h
@@ -0,0 +1,352 @@
+// Copyright Epic Games, Inc. All Rights Reserved.
+
+#pragma once
+
+#include "BaseTools/MultiSelectionMeshEditingTool.h"
+#include "DynamicMesh/DynamicMeshAttributeSet.h"
+#include "Graphs/GenerateStaticMeshLODProcess.h"
+#include "GenerateStaticMeshLODAssetTool.generated.h"
+
+class UCollisionGeometryVisualizationProperties;
+class UPhysicsObjectToolPropertySet;
+class UPreviewGeometry;
+namespace UE::Geometry { class IDynamicMeshOperatorFactory; }
+
+
+// predeclarations
+struct FMeshDescription;
+class UDynamicMeshComponent;
+class UMeshOpPreviewWithBackgroundCompute;
+class UGenerateStaticMeshLODAssetTool;
+class UStaticMeshLODGenerationSettings;
+namespace GenerateStaticMeshLODAssetLocals
+{
+	class FGenerateStaticMeshLODAssetOperatorFactory;
+}
+
+UENUM()
+enum class EGenerateLODAssetOutputMode : uint8
+{
+	CreateNewAsset = 0,
+	UpdateExistingAsset = 1
+};
+
+
+/**
+ * Tool builder
+ */
+UCLASS()
+class MESHLODTOOLSET_API UGenerateStaticMeshLODAssetToolBuilder : public UMultiSelectionMeshEditingToolBuilder
+{
+	GENERATED_BODY()
+
+public:
+	bool bUseAssetEditorMode = false;
+	bool bInRestrictiveMode = false;
+
+	virtual bool CanBuildTool(const FToolBuilderState& SceneState) const override;
+	virtual UMultiSelectionMeshEditingTool* CreateNewTool(const FToolBuilderState& SceneState) const override;
+
+protected:
+	virtual const FToolTargetTypeRequirements& GetTargetRequirements() const override;
+};
+
+
+
+UCLASS()
+class MESHLODTOOLSET_API UGenerateStaticMeshLODAssetToolOutputProperties : public UInteractiveToolPropertySet
+{
+	GENERATED_BODY()
+public:
+	/** Whether to modify the static mesh in place or create a new one. */
+	UPROPERTY(EditAnywhere, Category = "Output Options", meta = (TransientToolProperty, HideEditConditionToggle, EditCondition = "bShowOutputMode"))
+	EGenerateLODAssetOutputMode OutputMode = EGenerateLODAssetOutputMode::CreateNewAsset;
+
+	/** Base name for newly-generated asset */
+	UPROPERTY(EditAnywhere, Category = "Output Options", meta = (TransientToolProperty, EditConditionHides, EditCondition = "OutputMode == EGenerateLODAssetOutputMode::CreateNewAsset"))
+	FString NewAssetName;
+
+	/** If the Asset doesn't already have a HiRes source, store the input mesh as the HiRes source */
+	UPROPERTY(EditAnywhere, Category = "Output Options", meta = (EditConditionHides, EditCondition = "OutputMode == EGenerateLODAssetOutputMode::UpdateExistingAsset"))
+	bool bSaveInputAsHiResSource = true;
+
+	/** Suffix to append to newly-generated Asset (Meshes, Textures, Materials, etc) */
+	UPROPERTY(EditAnywhere, Category = "Output Options", meta = (TransientToolProperty))
+	FString GeneratedSuffix = TEXT("_AutoLOD");
+
+	/** If false, then OutputMode will not be shown in DetailsView panels (otherwise no effect) */
+	UPROPERTY(meta = (TransientToolProperty))
+	bool bShowOutputMode = true;
+};
+
+
+
+
+UENUM()
+enum class EGenerateLODAssetToolPresetAction : uint8
+{
+	ReadFromPreset,
+	WriteToPreset
+};
+
+
+UCLASS()
+class MESHLODTOOLSET_API UGenerateStaticMeshLODAssetToolPresetProperties : public UInteractiveToolPropertySet
+{
+	GENERATED_BODY()
+
+public:
+	TWeakObjectPtr<UGenerateStaticMeshLODAssetTool> ParentTool;
+	void Initialize(UGenerateStaticMeshLODAssetTool* ParentToolIn) { ParentTool = ParentToolIn; }
+	virtual void PostAction(EGenerateLODAssetToolPresetAction Action);
+
+public:
+	/** Preset Asset represents a set of Saved Settings for this Tool */
+	UPROPERTY(EditAnywhere, Category = Preset, meta = (DisplayName = "Settings Preset"))
+	TWeakObjectPtr<UStaticMeshLODGenerationSettings> Preset;
+
+	/** Read the current Tool settings from the Preset Asset */
+	UFUNCTION(CallInEditor, Category = Preset, meta = (DisplayPriority = 0))
+	void ReadFromPreset()
+	{
+		PostAction(EGenerateLODAssetToolPresetAction::ReadFromPreset);
+	}
+
+	/** Save the current Tool settings to the Preset Asset */
+	UFUNCTION(CallInEditor, Category = Preset, meta = (DisplayPriority = 1))
+	void WriteToPreset()
+	{
+		PostAction(EGenerateLODAssetToolPresetAction::WriteToPreset);
+	}
+};
+
+
+
+/**
+ * Standard properties
+ */
+UCLASS()
+class MESHLODTOOLSET_API UGenerateStaticMeshLODAssetToolProperties : public UInteractiveToolPropertySet
+{
+	GENERATED_BODY()
+
+public:
+	UPROPERTY(EditAnywhere, Category = "Generator Configuration")
+	FGenerateStaticMeshLODProcess_PreprocessSettings Preprocessing;
+
+	UPROPERTY(EditAnywhere, Category = "Generator Configuration", meta = (ExpandByDefault))
+	FGenerateStaticMeshLODProcessSettings MeshGeneration;
+
+	UPROPERTY(EditAnywhere, Category = "Generator Configuration", meta = (ExpandByDefault))
+	FGenerateStaticMeshLODProcess_SimplifySettings Simplification;
+
+	UPROPERTY(EditAnywhere, Category = "Generator Configuration", meta = (ExpandByDefault))
+	FGenerateStaticMeshLODProcess_NormalsSettings Normals;
+
+	UPROPERTY(EditAnywhere, Category = "Generator Configuration", meta = (ExpandByDefault))
+	FGenerateStaticMeshLODProcess_TextureSettings TextureBaking;
+
+	UPROPERTY(EditAnywhere, Category = "Generator Configuration")
+	FGenerateStaticMeshLODProcess_UVSettings UVGeneration;
+
+	UPROPERTY(EditAnywhere, Category = "Generator Configuration")
+	FGenerateStaticMeshLODProcess_CollisionSettings SimpleCollision;
+
+
+	// ------------
+	// Code copied from UPolygroupLayersProperties
+
+	/** Group layer to use for partitioning the mesh for simple collision generation */
+	UPROPERTY(EditAnywhere, Category = Settings, meta = (GetOptions = GetGroupLayersFunc))
+	FName CollisionGroupLayerName = TEXT("Default");
+
+	// this function is called provide set of available group layers
+	UFUNCTION()
+	TArray<FString> GetGroupLayersFunc()
+	{
+		return GroupLayersList;
+	}
+
+	// internal list used to implement above
+	UPROPERTY(meta = (TransientToolProperty))
+	TArray<FString> GroupLayersList;
+
+	void InitializeGroupLayers(const FDynamicMesh3* Mesh)
+	{
+		GroupLayersList.Reset();
+		GroupLayersList.Add(TEXT("Default"));		// always have standard group
+		if (Mesh->Attributes())
+		{
+			for (int32 k = 0; k < Mesh->Attributes()->NumPolygroupLayers(); k++)
+			{
+				FName Name = Mesh->Attributes()->GetPolygroupLayer(k)->GetName();
+				GroupLayersList.Add(Name.ToString());
+			}
+		}
+
+		if (GroupLayersList.Contains(CollisionGroupLayerName.ToString()) == false)		// discard restored value if it doesn't apply
+		{
+			CollisionGroupLayerName = FName(GroupLayersList[0]);
+		}
+	}
+
+};
+
+
+UENUM()
+enum class EGenerateStaticMeshLOD_BakeConstraint
+{
+	NoConstraint = 0,
+	DoNotBake = 1
+};
+
+USTRUCT()
+struct MESHLODTOOLSET_API FGenerateStaticMeshLOD_TextureConfig
+{
+	GENERATED_BODY()
+
+	UPROPERTY(VisibleAnywhere, Category = Texture)
+	TObjectPtr<UTexture2D> Texture = nullptr;
+
+	UPROPERTY(EditAnywhere, Category = Texture)
+	EGenerateStaticMeshLOD_BakeConstraint Constraint = EGenerateStaticMeshLOD_BakeConstraint::NoConstraint;
+
+	bool operator==(const FGenerateStaticMeshLOD_TextureConfig& Other) const
+	{
+		return Texture == Other.Texture && Constraint == Other.Constraint;
+	}
+};
+
+
+USTRUCT()
+struct MESHLODTOOLSET_API FGenerateStaticMeshLOD_MaterialConfig
+{
+	GENERATED_BODY()
+
+	UPROPERTY(VisibleAnywhere, Category = Material, meta = (TransientToolProperty, NoResetToDefault))
+	TObjectPtr<UMaterialInterface> Material = nullptr;
+
+	UPROPERTY(EditAnywhere, Category = Texture, meta = (TransientToolProperty, NoResetToDefault))
+	EGenerateStaticMeshLOD_BakeConstraint Constraint = EGenerateStaticMeshLOD_BakeConstraint::NoConstraint;
+
+	bool operator==(const FGenerateStaticMeshLOD_MaterialConfig& Other) const
+	{
+		return Material == Other.Material && Constraint == Other.Constraint;
+	}
+};
+
+
+
+UCLASS()
+class MESHLODTOOLSET_API UGenerateStaticMeshLODAssetToolTextureProperties : public UInteractiveToolPropertySet
+{
+	GENERATED_BODY()
+public:
+	UPROPERTY(EditAnywhere, Category = "Source Materials Configuration", meta = (TransientToolProperty, NoResetToDefault))
+	TArray<FGenerateStaticMeshLOD_MaterialConfig> Materials;
+
+	UPROPERTY(EditAnywhere, Category = "Source Textures Configuration", meta = (TransientToolProperty, NoResetToDefault))
+	TArray<FGenerateStaticMeshLOD_TextureConfig> Textures;
+
+	UPROPERTY(VisibleAnywhere, Category = "Baked Texture Previews", meta = (TransientToolProperty, NoResetToDefault))
+	TArray<TObjectPtr<UTexture2D>> PreviewTextures;
+};
+
+
+
+
+/**
+ * Simple tool to combine multiple meshes into a single mesh asset
+ */
+UCLASS()
+class MESHLODTOOLSET_API UGenerateStaticMeshLODAssetTool : public UMultiSelectionMeshEditingTool
+{
+	GENERATED_BODY()
+
+public:
+
+	// Enable UI Customization for running this Tool in the Static Mesh Asset Editor. Must be called before Setup.
+	virtual void SetUseAssetEditorMode(bool bEnable);
+
+	// Enable restrictive mode. Must be called before Setup.
+	virtual void SetRestrictiveMode(bool bEnable);
+
+	virtual void Setup() override;
+	virtual void OnShutdown(EToolShutdownType ShutdownType) override;
+
+	virtual void OnTick(float DeltaTime);
+
+	virtual bool HasCancel() const override { return true; }
+	virtual bool HasAccept() const override { return true; }
+	virtual bool CanAccept() const override;
+
+	virtual void RequestPresetAction(EGenerateLODAssetToolPresetAction ActionType);
+
+protected:
+
+	UPROPERTY()
+	TObjectPtr<UGenerateStaticMeshLODAssetToolOutputProperties> OutputProperties = nullptr;
+
+	UPROPERTY()
+	TObjectPtr<UGenerateStaticMeshLODAssetToolProperties> BasicProperties = nullptr;
+
+	UPROPERTY()
+	TObjectPtr<UGenerateStaticMeshLODAssetToolPresetProperties> PresetProperties = nullptr;
+
+	UPROPERTY()
+	TObjectPtr<UGenerateStaticMeshLODAssetToolTextureProperties> TextureProperties = nullptr;
+
+	UPROPERTY()
+	TObjectPtr<UCollisionGeometryVisualizationProperties> CollisionVizSettings = nullptr;
+
+
+
+protected:
+
+	friend class GenerateStaticMeshLODAssetLocals::FGenerateStaticMeshLODAssetOperatorFactory;
+
+	UPROPERTY()
+	TObjectPtr<UMeshOpPreviewWithBackgroundCompute> PreviewWithBackgroundCompute = nullptr;
+
+	UPROPERTY()
+	TArray<TObjectPtr<UTexture2D>> PreviewTextures;
+
+	UPROPERTY()
+	TArray<TObjectPtr<UMaterialInterface>> PreviewMaterials;
+
+
+protected:
+
+	UPROPERTY()
+	TObjectPtr<UPhysicsObjectToolPropertySet> ObjectData = nullptr;
+
+	UPROPERTY()
+	TObjectPtr<UPreviewGeometry> CollisionPreview;
+
+protected:
+	bool bIsInAssetEditorMode = false;
+	bool bRestrictiveMode = false;
+
+	UPROPERTY()
+	TObjectPtr<UGenerateStaticMeshLODProcess> GenerateProcess;
+
+	TUniquePtr<UE::Geometry::IDynamicMeshOperatorFactory> OpFactory;
+
+	void OnSettingsModified();
+
+	void CreateNewAsset();
+	void UpdateExistingAsset();
+
+	void OnPresetSelectionChanged();
+
+	bool ValidateSettings() const;
+};
+
+#if UE_ENABLE_INCLUDE_ORDER_DEPRECATED_IN_5_2
+#include "CoreMinimal.h"
+#include "Drawing/PreviewGeometryActor.h"
+#include "ModelingOperators.h" //IDynamicMeshOperatorFactory
+#include "Physics/CollisionPropertySets.h"
+#include "PreviewMesh.h"
+#include "UObject/NoExportTypes.h"
+#endif
diff --git a/Engine/Plugins/Editor/MeshLODToolset/Source/MeshLODToolset/Public/Tools/LODGenerationSettingsAsset.h b/Engine/Plugins/Editor/MeshLODToolset/Source/MeshLODToolset/Public/Tools/LODGenerationSettingsAsset.h
new file mode 100644
--- /dev/null
+++ b/Engine/Plugins/Editor/MeshLODToolset/Source/MeshLODToolset/Public/Tools/LODGenerationSettingsAsset.h
@@ -0,0 +1,51 @@
+// Copyright Epic Games, Inc. All Rights Reserved.
+
+#pragma once
+
+#include "Graphs/GenerateStaticMeshLODProcess.h"
+
+#include "LODGenerationSettingsAsset.generated.h"
+
+
+/**
+ * UStaticMeshLODGenerationSettings is intended to be a stored version of the settings used
+ * by UGenerateStaticMeshLODProcess (and the associated UGenerateStaticMeshLODAssetTool). 
+ * This UObject is exposed as an Asset type in the Editor via UStaticMeshLODGenerationSettingsFactory.
+ * 
+ * The Tool uses these serialized settings as a 'Preset', ie the user can save a set
+ * of configured settings, or load previously-saved settings. 
+ */
+UCLASS(Blueprintable)
+class MESHLODTOOLSET_API UStaticMeshLODGenerationSettings : public UObject
+{
+	GENERATED_BODY()
+
+#if WITH_EDITORONLY_DATA
+public:
+	UPROPERTY(EditAnywhere, Category=Settings)
+	FGenerateStaticMeshLODProcess_PreprocessSettings Preprocessing;
+
+	UPROPERTY(EditAnywhere, Category=Settings)
+	FGenerateStaticMeshLODProcessSettings MeshGeneration;
+
+	UPROPERTY(EditAnywhere, Category=Settings)
+	FGenerateStaticMeshLODProcess_SimplifySettings Simplification;
+
+	UPROPERTY(EditAnywhere, Category=Settings)
+	FGenerateStaticMeshLODProcess_NormalsSettings Normals;
+
+	UPROPERTY(EditAnywhere, Category=Settings)
+	FGenerateStaticMeshLODProcess_TextureSettings TextureBaking;
+
+	UPROPERTY(EditAnywhere, Category=Settings)
+	FGenerateStaticMeshLODProcess_UVSettings UVGeneration;
+
+	UPROPERTY(EditAnywhere, Category=Settings)
+	FGenerateStaticMeshLODProcess_CollisionSettings SimpleCollision;
+#endif
+};
+
+#if UE_ENABLE_INCLUDE_ORDER_DEPRECATED_IN_5_2
+#include "CoreMinimal.h"
+#include "EngineDefines.h"
+#endif
diff --git a/Engine/Plugins/Editor/MeshLODToolset/Source/MeshLODToolset/Public/Tools/LODGenerationSettingsFactory.h b/Engine/Plugins/Editor/MeshLODToolset/Source/MeshLODToolset/Public/Tools/LODGenerationSettingsFactory.h
new file mode 100644
--- /dev/null
+++ b/Engine/Plugins/Editor/MeshLODToolset/Source/MeshLODToolset/Public/Tools/LODGenerationSettingsFactory.h
@@ -0,0 +1,30 @@
+// Copyright Epic Games, Inc. All Rights Reserved.
+
+#pragma once
+
+#include "Factories/Factory.h"
+#include "LODGenerationSettingsFactory.generated.h"
+
+/**
+ * Asset Factory for UStaticMeshLODGenerationSettings, which is used to save settings for
+ * the AutoLOD Tool/Process as an Asset in the Editor
+ */
+UCLASS(MinimalAPI, hidecategories = Object)
+class UStaticMeshLODGenerationSettingsFactory : public UFactory
+{
+	GENERATED_UCLASS_BODY()
+
+	//~ Begin UFactory Interface
+	virtual UObject* FactoryCreateNew(UClass* Class, UObject* InParent, FName Name, EObjectFlags Flags, UObject* Context, FFeedbackContext* Warn) override;
+	virtual uint32 GetMenuCategories() const override;
+	virtual FText GetDisplayName() const override;
+	virtual FString GetDefaultNewAssetName() const override;
+	virtual bool ShouldShowInNewMenu() const override { return true; }
+	//~ Begin UFactory Interface	
+};
+
+
+
+#if UE_ENABLE_INCLUDE_ORDER_DEPRECATED_IN_5_2
+#include "CoreMinimal.h"
+#endif
diff --git a/Engine/Plugins/Editor/MeshLODToolset/Source/MeshLODToolset/Public/Tools/LODManagerTool.h b/Engine/Plugins/Editor/MeshLODToolset/Source/MeshLODToolset/Public/Tools/LODManagerTool.h
new file mode 100644
--- /dev/null
+++ b/Engine/Plugins/Editor/MeshLODToolset/Source/MeshLODToolset/Public/Tools/LODManagerTool.h
@@ -0,0 +1,265 @@
+// Copyright Epic Games, Inc. All Rights Reserved.
+
+#pragma once
+
+#include "BaseTools/MultiSelectionMeshEditingTool.h"
+#include "DynamicMesh/DynamicMesh3.h"
+#include "Engine/StaticMesh.h"
+#include "LODManagerTool.generated.h"
+
+class UPreviewGeometry;
+class UPreviewMesh;
+
+// predeclarations
+class UMaterialInterface;
+class ULODManagerTool;
+
+/**
+ *
+ */
+UCLASS()
+class MESHLODTOOLSET_API ULODManagerToolBuilder : public UMultiSelectionMeshEditingToolBuilder
+{
+	GENERATED_BODY()
+
+public:
+	virtual bool CanBuildTool(const FToolBuilderState& SceneState) const override;
+	virtual UMultiSelectionMeshEditingTool* CreateNewTool(const FToolBuilderState& SceneState) const override;
+
+protected:
+	virtual const FToolTargetTypeRequirements& GetTargetRequirements() const override;
+};
+
+
+USTRUCT()
+struct MESHLODTOOLSET_API FLODManagerLODInfo
+{
+	GENERATED_BODY()
+
+	UPROPERTY(VisibleAnywhere, Category = LODInformation, meta=(NoResetToDefault))
+	int32 VertexCount = 0;
+
+	UPROPERTY(VisibleAnywhere, Category = LODInformation, meta=(NoResetToDefault))
+	int32 TriangleCount = 0;
+};
+
+
+
+UCLASS()
+class MESHLODTOOLSET_API ULODManagerLODProperties : public UInteractiveToolPropertySet
+{
+	GENERATED_BODY()
+public:
+	UPROPERTY(VisibleAnywhere, Category = LODInformation, meta = (NoResetToDefault))
+	TArray<FLODManagerLODInfo> SourceLODs;
+
+	UPROPERTY(VisibleAnywhere, Category = LODInformation, meta = (NoResetToDefault))
+	TArray<FLODManagerLODInfo> HiResSource;
+
+
+	UPROPERTY(VisibleAnywhere, Category = LODInformation, meta = (NoResetToDefault))
+	TArray<FLODManagerLODInfo> RenderLODs;
+
+	UPROPERTY(VisibleAnywhere, Category = Nanite, meta = (NoResetToDefault, DisplayName="Enabled"))
+	bool bNaniteEnabled = false;
+
+	// Percentage of triangles kept by Nanite
+	UPROPERTY(VisibleAnywhere, Category = Nanite, meta = (NoResetToDefault))
+	float KeepTrianglePercent = 0.0f;
+
+	UPROPERTY(VisibleAnywhere, Category = Materials, meta = (NoResetToDefault))
+	TArray<FStaticMaterial> Materials;
+
+};
+
+
+
+UCLASS()
+class MESHLODTOOLSET_API ULODManagerPreviewLODProperties : public UInteractiveToolPropertySet
+{
+	GENERATED_BODY()
+public:
+	/** LOD to visualise. Default option is equivalent to disabling the Tool, RenderData is the mesh used for rendering derived from the SourceModel (possibly simplified) */
+	UPROPERTY(EditAnywhere, Category = LODPreview, meta = (DisplayName = "Show LOD", NoResetToDefault, GetOptions = GetLODNamesFunc))
+	FString VisibleLOD;
+
+	UFUNCTION()
+	const TArray<FString>& GetLODNamesFunc() const { return LODNamesList; }
+
+	UPROPERTY(meta = (TransientToolProperty))
+	TArray<FString> LODNamesList;
+
+	/** Control whether mesh borders are displayed */
+	UPROPERTY(EditAnywhere, DisplayName = "Show Borders", Category = LODPreview, meta = (EditCondition = "!bShowingDefaultLOD", HideEditConditionToggle))
+	bool bShowSeams = true;
+
+	// Used to expose the VisibleLOD state to the above EditCondition
+	UPROPERTY(Transient, meta = (TransientToolProperty))
+	bool bShowingDefaultLOD = true;
+};
+
+
+
+
+
+
+
+UENUM()
+enum class ELODManagerToolActions
+{
+	NoAction,
+
+	MoveHiResToLOD0,
+	DeleteHiResSourceModel,
+	RemoveUnreferencedMaterials
+};
+
+
+
+UCLASS()
+class MESHLODTOOLSET_API ULODManagerActionPropertySet : public UInteractiveToolPropertySet
+{
+	GENERATED_BODY()
+
+public:
+	TWeakObjectPtr<ULODManagerTool> ParentTool;
+
+	void Initialize(ULODManagerTool* ParentToolIn) { ParentTool = ParentToolIn; }
+	void PostAction(ELODManagerToolActions Action);
+
+};
+
+
+UCLASS()
+class MESHLODTOOLSET_API ULODManagerHiResSourceModelActions : public ULODManagerActionPropertySet
+{
+	GENERATED_BODY()
+public:
+	/** Move the HiRes Source Model to LOD0 */
+	UFUNCTION(CallInEditor, Category = HiResSourceModel, meta = (DisplayPriority = 0))
+	void MoveToLOD0()
+	{
+		PostAction(ELODManagerToolActions::MoveHiResToLOD0);
+	}
+
+	/** Delete the HiRes Source Model */
+	UFUNCTION(CallInEditor, Category = HiResSourceModel, meta = (DisplayPriority = 1))
+	void Delete()
+	{
+		PostAction(ELODManagerToolActions::DeleteHiResSourceModel);
+	}
+};
+
+
+UCLASS()
+class MESHLODTOOLSET_API ULODManagerMaterialActions : public ULODManagerActionPropertySet
+{
+	GENERATED_BODY()
+public:
+	/** Discard any Materials that are not referenced by any LOD */
+	UFUNCTION(CallInEditor, Category = MaterialSet, meta = (DisplayPriority = 0))
+	void CleanMaterials()
+	{
+		PostAction(ELODManagerToolActions::RemoveUnreferencedMaterials);
+	}
+};
+
+
+
+
+
+/**
+ * Mesh Attribute Editor Tool
+ */
+UCLASS()
+class MESHLODTOOLSET_API ULODManagerTool : public UMultiSelectionMeshEditingTool
+{
+	GENERATED_BODY()
+
+public:
+	ULODManagerTool();
+
+	virtual void Setup() override;
+	virtual void OnShutdown(EToolShutdownType ShutdownType) override;
+	virtual void OnTick(float DeltaTime) override;
+
+	virtual bool HasCancel() const override { return false; }
+	virtual bool HasAccept() const override { return false; }
+	virtual bool CanAccept() const override { return false; }
+
+	virtual void RequestAction(ELODManagerToolActions ActionType);
+
+protected:
+
+	UPROPERTY()
+	TObjectPtr<ULODManagerLODProperties> LODInfoProperties;
+
+	UPROPERTY()
+	TObjectPtr<ULODManagerPreviewLODProperties> LODPreviewProperties;
+
+	UPROPERTY()
+	TObjectPtr<ULODManagerHiResSourceModelActions> HiResSourceModelActions;
+
+	UPROPERTY()
+	TObjectPtr<ULODManagerMaterialActions> MaterialActions;
+
+
+
+public:
+	UFUNCTION()
+	void DeleteHiResSourceModel();
+
+	UFUNCTION()
+	void MoveHiResToLOD0();
+
+	UFUNCTION()
+	void RemoveUnreferencedMaterials();
+
+protected:
+	UStaticMesh* GetSingleStaticMesh();
+
+	ELODManagerToolActions PendingAction = ELODManagerToolActions::NoAction;
+	
+	bool bLODInfoValid = false;
+	void UpdateLODInfo();
+
+	struct FLODName
+	{
+		int32 SourceModelIndex = -1;
+		int32 RenderDataIndex = -1;
+		int32 OtherIndex = -1;
+		bool IsDefault() const { return SourceModelIndex == -1 && RenderDataIndex == -1 && OtherIndex == -1; }
+	};
+	FString DefaultLODName = FString("Default");
+	TMap<FString, FLODName> ActiveLODNames;
+	void UpdateLODNames();
+
+	struct FLODMeshInfo
+	{
+		UE::Geometry::FDynamicMesh3 Mesh;
+
+		bool bInfoCached = false;
+		TArray<int> BoundaryEdges;
+	};
+
+	TMap<FString, TUniquePtr<FLODMeshInfo>> LODMeshCache;
+
+	UPROPERTY()
+	TObjectPtr<UPreviewMesh> LODPreview;
+
+	UPROPERTY()
+	TObjectPtr<UPreviewGeometry> LODPreviewLines;
+
+	bool bPreviewLODValid = false;
+	void UpdatePreviewLOD();
+	void UpdatePreviewLines(FLODMeshInfo& LODMeshInfo);
+	void ClearPreviewLines();
+	bool CacheLODMesh(const FString& Name, FLODName LODName);
+};
+
+#if UE_ENABLE_INCLUDE_ORDER_DEPRECATED_IN_5_2
+#include "CoreMinimal.h"
+#include "Drawing/PreviewGeometryActor.h"
+#include "DynamicMesh/MeshTangents.h"
+#include "PreviewMesh.h"
+#endif
diff --git a/Engine/Plugins/Experimental/MeshLODToolset/MeshLODToolset.uplugin b/Engine/Plugins/Experimental/MeshLODToolset/MeshLODToolset.uplugin
deleted file mode 100644
diff --git a/Engine/Plugins/Experimental/MeshLODToolset/Source/MeshLODToolset/MeshLODToolset.Build.cs b/Engine/Plugins/Experimental/MeshLODToolset/Source/MeshLODToolset/MeshLODToolset.Build.cs
deleted file mode 100644
diff --git a/Engine/Plugins/Experimental/MeshLODToolset/Source/MeshLODToolset/Private/Graphs/GenerateMeshLODGraph.cpp b/Engine/Plugins/Experimental/MeshLODToolset/Source/MeshLODToolset/Private/Graphs/GenerateMeshLODGraph.cpp
deleted file mode 100644
diff --git a/Engine/Plugins/Experimental/MeshLODToolset/Source/MeshLODToolset/Private/Graphs/GenerateStaticMeshLODProcess.cpp b/Engine/Plugins/Experimental/MeshLODToolset/Source/MeshLODToolset/Private/Graphs/GenerateStaticMeshLODProcess.cpp
deleted file mode 100644
diff --git a/Engine/Plugins/Experimental/MeshLODToolset/Source/MeshLODToolset/Private/MeshLODToolsetModule.cpp b/Engine/Plugins/Experimental/MeshLODToolset/Source/MeshLODToolset/Private/MeshLODToolsetModule.cpp
deleted file mode 100644
diff --git a/Engine/Plugins/Experimental/MeshLODToolset/Source/MeshLODToolset/Private/Tools/DetailsCustomizations/AutoLODToolCustomizations.cpp b/Engine/Plugins/Experimental/MeshLODToolset/Source/MeshLODToolset/Private/Tools/DetailsCustomizations/AutoLODToolCustomizations.cpp
deleted file mode 100644
diff --git a/Engine/Plugins/Experimental/MeshLODToolset/Source/MeshLODToolset/Private/Tools/GenerateStaticMeshLODAssetTool.cpp b/Engine/Plugins/Experimental/MeshLODToolset/Source/MeshLODToolset/Private/Tools/GenerateStaticMeshLODAssetTool.cpp
deleted file mode 100644
diff --git a/Engine/Plugins/Experimental/MeshLODToolset/Source/MeshLODToolset/Private/Tools/LODGenerationSettingsAsset.cpp b/Engine/Plugins/Experimental/MeshLODToolset/Source/MeshLODToolset/Private/Tools/LODGenerationSettingsAsset.cpp
deleted file mode 100644
diff --git a/Engine/Plugins/Experimental/MeshLODToolset/Source/MeshLODToolset/Private/Tools/LODGenerationSettingsFactory.cpp b/Engine/Plugins/Experimental/MeshLODToolset/Source/MeshLODToolset/Private/Tools/LODGenerationSettingsFactory.cpp
deleted file mode 100644
diff --git a/Engine/Plugins/Experimental/MeshLODToolset/Source/MeshLODToolset/Private/Tools/LODManagerTool.cpp b/Engine/Plugins/Experimental/MeshLODToolset/Source/MeshLODToolset/Private/Tools/LODManagerTool.cpp
deleted file mode 100644
diff --git a/Engine/Plugins/Experimental/MeshLODToolset/Source/MeshLODToolset/Public/AssetDefinition_StaticMeshLODGenerationSettings.h b/Engine/Plugins/Experimental/MeshLODToolset/Source/MeshLODToolset/Public/AssetDefinition_StaticMeshLODGenerationSettings.h
deleted file mode 100644
diff --git a/Engine/Plugins/Experimental/MeshLODToolset/Source/MeshLODToolset/Public/Graphs/GenerateMeshLODGraph.h b/Engine/Plugins/Experimental/MeshLODToolset/Source/MeshLODToolset/Public/Graphs/GenerateMeshLODGraph.h
deleted file mode 100644
diff --git a/Engine/Plugins/Experimental/MeshLODToolset/Source/MeshLODToolset/Public/Graphs/GenerateStaticMeshLODProcess.h b/Engine/Plugins/Experimental/MeshLODToolset/Source/MeshLODToolset/Public/Graphs/GenerateStaticMeshLODProcess.h
deleted file mode 100644
diff --git a/Engine/Plugins/Experimental/MeshLODToolset/Source/MeshLODToolset/Public/MeshLODToolsetModule.h b/Engine/Plugins/Experimental/MeshLODToolset/Source/MeshLODToolset/Public/MeshLODToolsetModule.h
deleted file mode 100644
diff --git a/Engine/Plugins/Experimental/MeshLODToolset/Source/MeshLODToolset/Public/Tools/DetailsCustomizations/AutoLODToolCustomizations.h b/Engine/Plugins/Experimental/MeshLODToolset/Source/MeshLODToolset/Public/Tools/DetailsCustomizations/AutoLODToolCustomizations.h
deleted file mode 100644
diff --git a/Engine/Plugins/Experimental/MeshLODToolset/Source/MeshLODToolset/Public/Tools/GenerateStaticMeshLODAssetTool.h b/Engine/Plugins/Experimental/MeshLODToolset/Source/MeshLODToolset/Public/Tools/GenerateStaticMeshLODAssetTool.h
deleted file mode 100644
diff --git a/Engine/Plugins/Experimental/MeshLODToolset/Source/MeshLODToolset/Public/Tools/LODGenerationSettingsAsset.h b/Engine/Plugins/Experimental/MeshLODToolset/Source/MeshLODToolset/Public/Tools/LODGenerationSettingsAsset.h
deleted file mode 100644
diff --git a/Engine/Plugins/Experimental/MeshLODToolset/Source/MeshLODToolset/Public/Tools/LODGenerationSettingsFactory.h b/Engine/Plugins/Experimental/MeshLODToolset/Source/MeshLODToolset/Public/Tools/LODGenerationSettingsFactory.h
deleted file mode 100644
diff --git a/Engine/Plugins/Experimental/MeshLODToolset/Source/MeshLODToolset/Public/Tools/LODManagerTool.h b/Engine/Plugins/Experimental/MeshLODToolset/Source/MeshLODToolset/Public/Tools/LODManagerTool.h
deleted file mode 100644
