Geometry Tools:  Mesh LOD Toolset - fix bug in LODManager that could leave the mesh hidden after deleting the highres mesh.  Also made sure the displayed information was correctly updated after undo in the tool

diff --git a/Engine/Plugins/Editor/MeshLODToolset/Source/MeshLODToolset/MeshLODToolset.Build.cs b/Engine/Plugins/Editor/MeshLODToolset/Source/MeshLODToolset/MeshLODToolset.Build.cs
--- a/Engine/Plugins/Editor/MeshLODToolset/Source/MeshLODToolset/MeshLODToolset.Build.cs
+++ b/Engine/Plugins/Editor/MeshLODToolset/Source/MeshLODToolset/MeshLODToolset.Build.cs
@@ -34,6 +34,7 @@ public class MeshLODToolset : ModuleRules
 				"MeshDescription",
 				"StaticMeshDescription",
 				"ModelingComponents",
+				"ModelingComponentsEditorOnly",
 				"MeshModelingToolsExp",
 				"GeometryFlowCore",
 				"GeometryFlowMeshProcessing",
diff --git a/Engine/Plugins/Editor/MeshLODToolset/Source/MeshLODToolset/Private/Tools/LODManagerTool.cpp b/Engine/Plugins/Editor/MeshLODToolset/Source/MeshLODToolset/Private/Tools/LODManagerTool.cpp
--- a/Engine/Plugins/Editor/MeshLODToolset/Source/MeshLODToolset/Private/Tools/LODManagerTool.cpp
+++ b/Engine/Plugins/Editor/MeshLODToolset/Source/MeshLODToolset/Private/Tools/LODManagerTool.cpp
@@ -24,6 +24,7 @@
 #include "TargetInterfaces/StaticMeshBackedTarget.h"
 #include "ToolContextInterfaces.h"
 #include "ToolTargetManager.h"
+#include "ToolTargets/StaticMeshComponentToolTarget.h"
 
 #include UE_INLINE_GENERATED_CPP_BY_NAME(LODManagerTool)
 
@@ -179,6 +180,26 @@ void ULODManagerTool::OnTick(float DeltaTime)
 	}
 	PendingAction = ELODManagerToolActions::NoAction;
 
+	// test capture a dirty state that results from undo while in the tool.
+	auto IsLODInfoDirty = [&]()
+		{
+			UStaticMesh* StaticMesh = GetSingleStaticMesh();
+			if (!StaticMesh || !LODInfoProperties)
+			{
+				return false;
+			}
+
+			// have we changed the number of materials
+			bool bDirty = (LODInfoProperties->Materials.Num() != StaticMesh->GetStaticMaterials().Num());
+
+			// what about hi-res mesh?
+			bDirty = bDirty || (StaticMesh->IsHiResMeshDescriptionValid() ? LODInfoProperties->HiResSource.Num() == 0 : LODInfoProperties->HiResSource.Num() != 0);
+
+			return bDirty;
+		};
+
+	bLODInfoValid = bLODInfoValid && !IsLODInfoDirty();
+
 	if (bLODInfoValid == false)
 	{
 		UpdateLODInfo();
@@ -199,6 +220,19 @@ UStaticMesh* ULODManagerTool::GetSingleStaticMesh()
 		return nullptr;
 	}
 	UStaticMeshComponent* StaticMeshComponent = Cast<UStaticMeshComponent>(UE::ToolTarget::GetTargetComponent(Targets[0]));
+	if (StaticMeshComponent == nullptr)
+	{
+		// We may have removed the highest quality LOD ( i.e hi-res ) during tool operation or restored it by doing an in-tool undo,
+		// so here we reset the target to current max quality.  Why? Because if the tool target 'Editing LOD' does not correspond to an existing LOD on 
+		// the mesh then target is considered invalid  and GetTargetComponent will fail as will showing the source object on shutdown.
+		if (UStaticMeshComponentToolTarget* StaticMeshComponentTarget = Cast<UStaticMeshComponentToolTarget>(Targets[0]))
+		{
+			StaticMeshComponentTarget->SetEditingLOD(EMeshLODIdentifier::MaxQuality);
+
+			StaticMeshComponent = Cast<UStaticMeshComponent>(UE::ToolTarget::GetTargetComponent(Targets[0]));
+		}	
+	}
+
 	if (StaticMeshComponent == nullptr || StaticMeshComponent->GetStaticMesh() == nullptr)
 	{
 		return nullptr;
@@ -270,6 +304,8 @@ void ULODManagerTool::UpdateLODInfo()
 
 
 	UpdateLODNames();
+
+	bPreviewLODValid = false;   
 }
 
 
@@ -323,8 +359,6 @@ void ULODManagerTool::UpdateLODNames()
 		}
 	}
 
-
-	bPreviewLODValid = false;    // should not always be necessary...
 }
 
 
@@ -392,6 +426,10 @@ void ULODManagerTool::UpdatePreviewLOD()
 		if (bNewValue != LODPreviewProperties->bShowingDefaultLOD)
 		{
 			LODPreviewProperties->bShowingDefaultLOD = bNewValue;
+			if (bNewValue == true)
+			{
+				LODPreviewProperties->VisibleLOD = DefaultLODName;
+			}
 			NotifyOfPropertyChangeByTool(LODPreviewProperties);
 		}
 	};
@@ -503,6 +541,8 @@ bool ULODManagerTool::CacheLODMesh(const FString& Name, FLODName LODName)
 void ULODManagerTool::DeleteHiResSourceModel()
 {
 	GetToolManager()->BeginUndoTransaction(LOCTEXT("DeleteHiResSourceModel", "Delete HiRes Source"));
+
+	const bool bComponentsHiddenByTool = LODPreview->IsVisible();
 	for (int32 ComponentIdx = 0; ComponentIdx < Targets.Num(); ComponentIdx++)
 	{
 		UStaticMeshComponent* StaticMeshComponent = Cast<UStaticMeshComponent>(UE::ToolTarget::GetTargetComponent(Targets[ComponentIdx]));
@@ -514,6 +554,13 @@ void ULODManagerTool::DeleteHiResSourceModel()
 
 		if (StaticMesh->IsHiResMeshDescriptionValid())
 		{
+			
+			if (bComponentsHiddenByTool)
+			{ 
+				// temporarily set the visible flag - that way an undo after the tool is closed will restore a visible mesh
+				UE::ToolTarget::ShowSourceObject(Targets[ComponentIdx]);
+			}
+
 			StaticMesh->Modify();
 
 			StaticMesh->ModifyHiResMeshDescription();
@@ -521,6 +568,12 @@ void ULODManagerTool::DeleteHiResSourceModel()
 			StaticMesh->CommitHiResMeshDescription();
 
 			StaticMesh->PostEditChange();
+
+			// restore the visibility state expected by the tool
+			if (bComponentsHiddenByTool)
+			{
+				UE::ToolTarget::HideSourceObject(Targets[ComponentIdx]);
+			}
 		}
 	}
 
@@ -535,6 +588,8 @@ void ULODManagerTool::DeleteHiResSourceModel()
 void ULODManagerTool::MoveHiResToLOD0()
 {
 	GetToolManager()->BeginUndoTransaction(LOCTEXT("RestoreLOD0", "Move HiRes to LOD0"));
+
+	const bool bComponentsHiddenByTool = LODPreview->IsVisible();
 	for (int32 ComponentIdx = 0; ComponentIdx < Targets.Num(); ComponentIdx++)
 	{
 		UStaticMeshComponent* StaticMeshComponent = Cast<UStaticMeshComponent>(UE::ToolTarget::GetTargetComponent(Targets[ComponentIdx]));
@@ -546,6 +601,13 @@ void ULODManagerTool::MoveHiResToLOD0()
 
 		if (StaticMesh->IsHiResMeshDescriptionValid())
 		{
+
+			if (bComponentsHiddenByTool)
+			{
+				// temporarily set the visible flag - that way an undo after the tool is closed will restore a visible mesh
+				UE::ToolTarget::ShowSourceObject(Targets[ComponentIdx]);
+			}
+
 			StaticMesh->Modify();
 
 			StaticMesh->ModifyHiResMeshDescription();
@@ -561,6 +623,12 @@ void ULODManagerTool::MoveHiResToLOD0()
 			StaticMesh->CommitMeshDescription(0);
 
 			StaticMesh->PostEditChange();
+
+			// restore the visibility state expected by the tool
+			if (bComponentsHiddenByTool)
+			{
+				UE::ToolTarget::HideSourceObject(Targets[ComponentIdx]);
+			}
 		}
 	}
 
