Modeling Tools: Fix bug in 'Clean Materials' which was removing incorrect material

diff --git a/Engine/Plugins/Editor/MeshLODToolset/Source/MeshLODToolset/Private/Tools/LODManagerTool.cpp b/Engine/Plugins/Editor/MeshLODToolset/Source/MeshLODToolset/Private/Tools/LODManagerTool.cpp
--- a/Engine/Plugins/Editor/MeshLODToolset/Source/MeshLODToolset/Private/Tools/LODManagerTool.cpp
+++ b/Engine/Plugins/Editor/MeshLODToolset/Source/MeshLODToolset/Private/Tools/LODManagerTool.cpp
@@ -672,33 +672,33 @@ void ULODManagerTool::RemoveUnreferencedMaterials()
 		}
 		int32 NumMeshes = Meshes.Num();
 
-
+		// using RenderData instead of Static Mesh Attributes because not every Source Model will have
+		// a Mesh Description to inspect, which caused a crash
+		FStaticMeshRenderData* RenderData = StaticMesh->GetRenderData();
+		
 		// for each mesh, collect list of material indices, and then set MatUsedFlags
 		TArray<TArray<int32>> MeshMaterialIndices;
 		MeshMaterialIndices.SetNum(NumMeshes);
-		for (int32 mi = 0; mi < NumMeshes; ++mi)
+		
+		for (int32 lod = 0; lod < NumMeshes; ++lod)
 		{
-			TMap<FPolygonGroupID, int32> PolygonGroupToSectionIndex;
-			TMap<FPolygonGroupID, int32> PolygonGroupToMaterialIndex;
-			FStaticMeshConstAttributes Attributes(*Meshes[mi]);
-			TPolygonGroupAttributesConstRef<FName> PolygonGroupImportedMaterialSlotNames = Attributes.GetPolygonGroupMaterialSlotNames();
-			for (FPolygonGroupID PolygonGroupID : Meshes[mi]->PolygonGroups().GetElementIDs() )
+			if (RenderData && RenderData->LODResources.IsValidIndex(lod))
 			{
-				int32& SectionIndex = PolygonGroupToSectionIndex.FindOrAdd(PolygonGroupID);
-				int32 MaterialIndex = StaticMesh->GetMaterialIndexFromImportedMaterialSlotName(PolygonGroupImportedMaterialSlotNames[PolygonGroupID]);
-				if (MaterialIndex == INDEX_NONE)
+				FStaticMeshLODResources& LOD = RenderData->LODResources[lod];
+				int NumSections = LOD.Sections.Num();
+
+				for (int32 SectionIndex = 0; SectionIndex < NumSections; SectionIndex++)
 				{
-					MaterialIndex = PolygonGroupID.GetValue();
+					// For every section in the LOD, retrieves the index of the material used
+					FMeshSectionInfo Info = StaticMesh->GetSectionInfoMap().Get(lod, SectionIndex);
+					int32 MaterialIndex = Info.MaterialIndex;
+					if (StaticMesh->GetStaticMaterials().IsValidIndex(MaterialIndex))
+					{
+						// material is used, set array value at its index to true
+						MatUsedFlags[MaterialIndex] = true;
+						MeshMaterialIndices[lod].AddUnique(MaterialIndex);
+					}
 				}
-				PolygonGroupToMaterialIndex.Add(PolygonGroupID, MaterialIndex);
-			}
-
-			for (const FTriangleID TriangleID : Meshes[mi]->Triangles().GetElementIDs())
-			{
-				const FPolygonGroupID PolygonGroupID = Meshes[mi]->GetTrianglePolygonGroup(TriangleID);
-				const int32 MaterialIndex = PolygonGroupToMaterialIndex[PolygonGroupID];
-				MatUsedFlags[MaterialIndex] = true;
-				MeshMaterialIndices[mi].AddUnique(MaterialIndex);
 			}
 		}
 
@@ -737,10 +737,13 @@ void ULODManagerTool::RemoveUnreferencedMaterials()
 		{
 			if (StaticMesh->IsSourceModelValid(k))
 			{
-				StaticMesh->ModifyMeshDescription(k);
 				FMeshDescription* Mesh = StaticMesh->GetMeshDescription(k);
-				Mesh->RemapPolygonGroups(RemapPolygonGroups);
-				StaticMesh->CommitMeshDescription(k);
+				if (Mesh)
+				{
+					StaticMesh->ModifyMeshDescription(k);
+					Mesh->RemapPolygonGroups(RemapPolygonGroups);
+					StaticMesh->CommitMeshDescription(k);
+				}
 			}
 		}
 		if (StaticMesh->IsHiResMeshDescriptionValid())
