make autolod tool set static mesh to transactional before modifying it to fix undo

diff --git a/Engine/Plugins/Editor/MeshLODToolset/Source/MeshLODToolset/Private/Graphs/GenerateStaticMeshLODProcess.cpp b/Engine/Plugins/Editor/MeshLODToolset/Source/MeshLODToolset/Private/Graphs/GenerateStaticMeshLODProcess.cpp
--- a/Engine/Plugins/Editor/MeshLODToolset/Source/MeshLODToolset/Private/Graphs/GenerateStaticMeshLODProcess.cpp
+++ b/Engine/Plugins/Editor/MeshLODToolset/Source/MeshLODToolset/Private/Graphs/GenerateStaticMeshLODProcess.cpp
@@ -1702,6 +1702,8 @@ void UGenerateStaticMeshLODProcess::UpdateSourceStaticMeshAsset(bool bSetNewHDSo
 {
 	GEditor->BeginTransaction(LOCTEXT("UpdateExistingAssetMessage", "Added Generated LOD"));
 
+	// make sure transactional flag is on for this asset
+	SourceStaticMesh->SetFlags(RF_Transactional);
 	SourceStaticMesh->Modify();
 
 	FStaticMeshSourceModel& SrcModel = SourceStaticMesh->GetSourceModel(0);
