Modeling Tools: LODManager.  Fix issue when undoing after exiting the tool can leave the mesh visibility set to off.

diff --git a/Engine/Plugins/Editor/MeshLODToolset/Source/MeshLODToolset/Private/Tools/LODManagerTool.cpp b/Engine/Plugins/Editor/MeshLODToolset/Source/MeshLODToolset/Private/Tools/LODManagerTool.cpp
--- a/Engine/Plugins/Editor/MeshLODToolset/Source/MeshLODToolset/Private/Tools/LODManagerTool.cpp
+++ b/Engine/Plugins/Editor/MeshLODToolset/Source/MeshLODToolset/Private/Tools/LODManagerTool.cpp
@@ -643,6 +643,17 @@ void ULODManagerTool::MoveHiResToLOD0()
 
 void ULODManagerTool::RemoveUnreferencedMaterials()
 {
+	const bool bComponentsHiddenByTool = LODPreview->IsVisible();
+
+	// temporarily set the visible flag - that way an undo after the tool is closed will restore a visible mesh
+	if (bComponentsHiddenByTool)
+	{
+		for (int32 ComponentIdx = 0; ComponentIdx < Targets.Num(); ComponentIdx++)
+		{
+			UE::ToolTarget::ShowSourceObject(Targets[ComponentIdx]);
+		}
+	}
+
 	GetToolManager()->BeginUndoTransaction(LOCTEXT("RemoveUnreferencedMaterials", "Remove Unreferenced Materials"));
 	for (int32 ComponentIdx = 0; ComponentIdx < Targets.Num(); ComponentIdx++)
 	{
@@ -758,6 +769,16 @@ void ULODManagerTool::RemoveUnreferencedMaterials()
 	}
 
 	GetToolManager()->EndUndoTransaction();
+
+	// revert to the original visibility
+	if (bComponentsHiddenByTool)
+	{
+		for (int32 ComponentIdx = 0; ComponentIdx < Targets.Num(); ComponentIdx++)
+		{
+			UE::ToolTarget::HideSourceObject(Targets[ComponentIdx]);
+		}
+	}
+
 	bLODInfoValid = false;
 }
 
