Modeling Tools : LODManager Tool - rewrite so the target static mesh is only modified during accept on tool shutdown.  Previsously every internal tool action modified the static mesh counter to the modeling tools convention ( this introduced potential undo/redo weirdness).

diff --git a/Engine/Plugins/Editor/MeshLODToolset/Source/MeshLODToolset/Private/Tools/LODManagerTool.cpp b/Engine/Plugins/Editor/MeshLODToolset/Source/MeshLODToolset/Private/Tools/LODManagerTool.cpp
--- a/Engine/Plugins/Editor/MeshLODToolset/Source/MeshLODToolset/Private/Tools/LODManagerTool.cpp
+++ b/Engine/Plugins/Editor/MeshLODToolset/Source/MeshLODToolset/Private/Tools/LODManagerTool.cpp
@@ -69,6 +69,609 @@ void ULODManagerActionPropertySet::PostAction(ELODManagerToolActions Action)
 	}
 }
 
+namespace UE {
+namespace Geometry {
+namespace LODManagerHelper
+{
+
+	// Static functions used to query a UStaticMesh and generate FDynamicMesh3 representations of the LODs 
+	struct FMeshLODUtils
+	{
+		// return true if the static mesh has an imported mesh at the specified LODIndex
+		static bool HasImportedLOD(const UStaticMesh& StaticMesh, int32 LODIndex)
+		{
+			const FMeshDescription* LODMesh = StaticMesh.GetMeshDescription(LODIndex);
+			return (LODMesh != nullptr);
+		}
+		
+		// get the specified imported LOD.  the result will be a invalid SharedPtr if the given static mesh does not have
+		// an imported LOD at the specified index.
+		static TSharedPtr<FDynamicMesh3> GetImportedLOD(const UStaticMesh& StaticMesh, int32 LODIndex)
+		{
+			TSharedPtr<FDynamicMesh3> Result;
+			if (FMeshDescription* LODMeshDescription = StaticMesh.GetMeshDescription(LODIndex)) 
+			{
+				Result = MakeShared<FDynamicMesh3>();
+				// TODO: we only need to do this copy if normals or tangents are set to auto-generated...
+				FMeshBuildSettings LODBuildSettings = StaticMesh.GetSourceModel(LODIndex).BuildSettings;
+				FMeshDescription TmpMeshDescription(*LODMeshDescription);
+				UE::MeshDescription::InitializeAutoGeneratedAttributes(TmpMeshDescription, &LODBuildSettings);
+
+				FMeshDescriptionToDynamicMesh Converter;
+				Converter.Convert(&TmpMeshDescription, *Result, true);
+			}
+			return Result;
+		}
+
+		// return true if the StaticMesh has render data for the indicated LODIndex 
+		static bool HasRenderLOD(const UStaticMesh& StaticMesh, int32 LODIndex)
+		{
+			if (!StaticMesh.HasValidRenderData())
+			{
+				return false;
+			}
+			const FStaticMeshRenderData* RenderData = StaticMesh.GetRenderData();
+			return RenderData->LODResources.IsValidIndex(LODIndex);
+		}
+
+		// get a dynamic mesh computed from the render data at the specified LODIndex. will return invalid SharedPtr if the
+		// LOD doesn't exist.
+		static TSharedPtr<FDynamicMesh3> GetRenderLOD(const UStaticMesh& StaticMesh, int32 LODIndex)
+		{
+			TSharedPtr<FDynamicMesh3> Result;
+			if (!FMeshLODUtils::HasRenderLOD(StaticMesh, LODIndex))
+			{
+				return Result;
+			}
+
+			Result = MakeShared<FDynamicMesh3>();
+			const FStaticMeshRenderData* RenderData = StaticMesh.GetRenderData();
+			const FStaticMeshLODResources& LODResource = RenderData->LODResources[LODIndex];
+			UE::Conversion::RenderBuffersToDynamicMesh(LODResource.VertexBuffers, LODResource.IndexBuffer, LODResource.Sections, *Result);
+
+			return Result;
+		}
+
+		// return true if the static mesh has a "hi-res" LOD ( e.g. source for nanite build)
+		static bool HasHiResLOD(const UStaticMesh& StaticMesh)
+		{
+			return StaticMesh.IsHiResMeshDescriptionValid();
+		}
+
+		// return a dynamic mesh version of the high-res mesh if it exists (otherwise null )
+		static TSharedPtr<FDynamicMesh3> GetHiResLOD(const UStaticMesh& StaticMesh)
+		{
+			TSharedPtr<FDynamicMesh3> Result;
+
+			if (HasHiResLOD(StaticMesh))
+			{
+
+				if (FMeshDescription* HiResMeshDescription = StaticMesh.GetHiResMeshDescription())
+				{
+					// TODO: we only need to do this copy if normals or tangents are set to auto-generated...
+					FMeshBuildSettings HiResBuildSettings = StaticMesh.GetHiResSourceModel().BuildSettings;
+					FMeshDescription TmpMeshDescription(*HiResMeshDescription);
+					UE::MeshDescription::InitializeAutoGeneratedAttributes(TmpMeshDescription, &HiResBuildSettings);
+					{
+						Result = MakeShared<FDynamicMesh3>();
+						FMeshDescriptionToDynamicMesh Converter;
+						Converter.Convert(&TmpMeshDescription, *Result, true);
+					}
+				
+				}
+			}
+			return Result;
+		}
+
+		// return a dynamic mesh version of the high-res mesh converted to render data ( if it exists, otherwise null )
+		static TSharedPtr<FDynamicMesh3> GetHiResLODAsRenderData(const UStaticMesh& StaticMesh, int32& VertexCount, int32& TriangleCount)
+		{
+			VertexCount = -1;
+			TriangleCount = -1;
+
+			TSharedPtr<FDynamicMesh3> Result;
+
+			if (HasHiResLOD(StaticMesh))
+			{
+				if (FMeshDescription* HiResMeshDescription = StaticMesh.GetHiResMeshDescription())
+				{
+					UStaticMesh* TmpStaticMesh = NewObject<UStaticMesh>(GetTransientPackage());
+					TmpStaticMesh->AddSourceModel();
+					if (FMeshDescription* LOD0MeshDescription = TmpStaticMesh->CreateMeshDescription(0))
+					{
+						*LOD0MeshDescription = *HiResMeshDescription;
+						TmpStaticMesh->CommitMeshDescription(0);
+						TmpStaticMesh->PostEditChange(); // trigger build of the render data.
+
+						Result = MakeShared<FDynamicMesh3>();
+						const FStaticMeshRenderData* RenderData = TmpStaticMesh->GetRenderData();
+						const FStaticMeshLODResources& LODResource = RenderData->LODResources[0];
+						
+						VertexCount = LODResource.VertexBuffers.PositionVertexBuffer.GetNumVertices();
+						TriangleCount = LODResource.IndexBuffer.GetNumIndices() / 3;
+
+						UE::Conversion::RenderBuffersToDynamicMesh(LODResource.VertexBuffers, LODResource.IndexBuffer, LODResource.Sections, *Result);
+					}
+				}
+			}
+
+			return Result;
+		}
+	};
+
+
+	// Cache's FDynamicMesh3 representations of the imported LODs and RenderData LODs (for display use in the tool)
+	// as well as the triangle and vertex counts for the StaticMesh LODs.
+	class FDynamicMeshLODCache
+	{
+	public:
+
+		struct FMeshElementInfo
+		{
+			int32 VertexCount = -1;
+			int32 TriangleCount = -1;
+			bool IsValid() const {return !(VertexCount == -1 || TriangleCount == -1);}
+		};
+		
+		
+		FDynamicMeshLODCache(UStaticMesh& StaticMeshIn)
+		: StaticMesh(&StaticMeshIn)
+		{
+			MaxLOD = StaticMesh->GetNumSourceModels();
+			ResetLODMeshInfo();
+		}
+
+		// compute and cache the triangle and vertex counts for the original source meshes stored in the UStaticMesh.
+		void ResetLODMeshInfo()
+		{
+			LODToMeshInfo.Reset();
+			RenderLODToMeshInfo.Reset();
+
+
+			int32 NumSourceModels = StaticMesh->GetNumSourceModels();
+			for (int32 si = 0; si < NumSourceModels; ++si)
+			{
+				
+				if (FMeshDescription* LODMesh = StaticMesh->GetMeshDescription(si))		// generated LODs do not have mesh description...is that the only way to tell?
+				{
+					FMeshElementInfo LODMeshInfo = { LODMesh->Vertices().Num(), LODMesh->Triangles().Num() };
+					
+					TTuple<int32, FMeshElementInfo> KeyValuePair = {si, LODMeshInfo};
+					LODToMeshInfo.Add(KeyValuePair);
+				}
+			}
+
+
+			if (StaticMesh->HasValidRenderData())
+			{
+				const FStaticMeshRenderData* RenderData = StaticMesh->GetRenderData();
+				int NumRenderLODs = RenderData->LODResources.Num();
+				for (int32 ri = 0; ri < NumRenderLODs; ++ri)
+				{
+					const FStaticMeshLODResources& LODResource = RenderData->LODResources[ri];
+					
+					FMeshElementInfo LODMeshInfo;
+					LODMeshInfo.VertexCount = LODResource.VertexBuffers.PositionVertexBuffer.GetNumVertices();
+					LODMeshInfo.TriangleCount = LODResource.IndexBuffer.GetNumIndices() / 3;
+					TTuple<int32, FMeshElementInfo> KeyValuePair = { ri, LODMeshInfo };
+					RenderLODToMeshInfo.Add(KeyValuePair);
+				}
+			}
+
+			if (HasHiResLOD())
+			{
+				const FMeshDescription* HiResMesh = StaticMesh->GetHiResMeshDescription();
+				HiResMeshInfo = { HiResMesh->Vertices().Num(), HiResMesh->Triangles().Num() };
+			}
+		}
+
+
+
+		// -- Imported LODs --
+		bool HasImportedLOD(int32 LODIndex) const
+		{
+			return FMeshLODUtils::HasImportedLOD(*StaticMesh, LODIndex);
+		}
+
+		FMeshElementInfo GetLODMeshElementInfo(int32 LODIndex) const
+		{
+			if (LODToMeshInfo.Contains(LODIndex))
+			{
+				return LODToMeshInfo[LODIndex];
+			}
+			return  FMeshElementInfo();
+		}
+
+		// get DynamicMesh representation of the imported LOD.
+		// returns null ptr if the the LOD was not an imported mesh.
+		TSharedPtr<FDynamicMesh3> GetImportedLOD(int32 LODIndex)
+		{
+			TSharedPtr<FDynamicMesh3> Result;
+			if (!HasImportedLOD(LODIndex))
+			{
+				return Result;
+			}
+
+			if (!LODToDynamicMesh.Contains(LODIndex))
+			{
+
+				Result = FMeshLODUtils::GetImportedLOD(*StaticMesh, LODIndex);
+				if (Result.IsValid())
+				{
+					LODToDynamicMesh.Add(LODIndex, Result);
+				}				
+			}
+
+			if (LODToDynamicMesh.Contains(LODIndex))
+			{
+				Result = LODToDynamicMesh[LODIndex];
+			}
+
+			return Result;
+		}
+
+		// -- Imported HiRes Mesh
+		bool HasHiResLOD() const
+		{
+			return FMeshLODUtils::HasHiResLOD(*StaticMesh);
+		}
+		// get triangle and vertex count for hi-res mesh. Returns invalid Info if no hi-res mesh.
+		const FMeshElementInfo& GetHiResMeshInfo() const
+		{
+			return HiResMeshInfo;
+		}
+		// return DynamiMesh representation of the HiRes mesh if it exists, otherwise an invalid SharedPtr
+		TSharedPtr<FDynamicMesh3> GetHiResLOD()
+		{
+			if (!HiResDynamicMesh)
+			{
+				HiResDynamicMesh = FMeshLODUtils::GetHiResLOD(*StaticMesh);
+			}
+
+			return HiResDynamicMesh;
+		}
+
+		// -- Render LODs
+		bool HasRenderLOD(int32 LODIndex) const
+		{
+			return FMeshLODUtils::HasRenderLOD(*StaticMesh, LODIndex);
+		}
+
+		// get triangle and vertex count for specified RenderData mesh if it exists, otherwise an invalid info.
+		FMeshElementInfo GetRenderLODMeshInfo(int32 LODIndex) const
+		{
+			if (RenderLODToMeshInfo.Contains(LODIndex))
+			{
+				return RenderLODToMeshInfo[LODIndex];
+			}
+
+			return FMeshElementInfo();
+		}
+
+		// return DynamicMesh representation of the specified RenderData LOD if it exists, otherwise an invalid SharedPtr
+		TSharedPtr<FDynamicMesh3> GetRenderLOD(int32 LODIndex)
+		{
+			TSharedPtr<FDynamicMesh3> Result;
+			
+			if (!HasRenderLOD(LODIndex))
+			{
+				return Result;
+			}
+
+			if (!RenderLODToDynamicMesh.Contains(LODIndex))
+			{
+				Result = FMeshLODUtils::GetRenderLOD(*StaticMesh, LODIndex);
+				
+				if (Result.IsValid())
+				{
+					RenderLODToDynamicMesh.Add(LODIndex, Result);
+				}
+			}
+
+			if (RenderLODToDynamicMesh.Contains(LODIndex))
+			{
+				Result = RenderLODToDynamicMesh[LODIndex];
+			}
+
+			return Result;
+		}
+
+		// returns DynamicMesh representation of the HiRes Mesh after it is converted to a RenderData form. 
+		// will return invalid SharedPtr if no hi-res mesh exists.
+		TSharedPtr<FDynamicMesh3> GetHiResLODAsRenderData()
+		{
+			if (!HiResRenderDynamicMesh)
+			{
+				HiResRenderDynamicMesh = FMeshLODUtils::GetHiResLODAsRenderData(*StaticMesh, HiResRenderMeshInfo.VertexCount, HiResRenderMeshInfo.TriangleCount);
+			}
+
+			return HiResRenderDynamicMesh;
+		}
+		// returns the triangle and vertex count of the render data representation of the hi-res mesh. 
+		// will return invalid info if no hi-res exists.
+		FMeshElementInfo GetHiResLODAsRenderDataMeshInfo()
+		{
+			if (!HiResRenderMeshInfo.IsValid() && HasHiResLOD())
+			{
+				// trigger build and cache of hi-res as renderdata, this captures the needed info
+				GetHiResLODAsRenderData();
+			}
+
+			return HiResRenderMeshInfo;
+		}
+
+	protected:
+
+		int32 MaxLOD;
+		// caches for imported LODs
+		TMap<int32, FMeshElementInfo>  LODToMeshInfo;
+		TMap<int32, TSharedPtr<FDynamicMesh3>> LODToDynamicMesh;
+		
+		// caches for render data LODs
+		TMap<int32, FMeshElementInfo>  RenderLODToMeshInfo;
+		TMap<int32, TSharedPtr<FDynamicMesh3>> RenderLODToDynamicMesh;
+		
+		// for imported hi-res
+		FMeshElementInfo HiResMeshInfo;
+		TSharedPtr<FDynamicMesh3> HiResDynamicMesh;
+
+		// for render data hi-res 
+		FMeshElementInfo HiResRenderMeshInfo;
+		TSharedPtr<FDynamicMesh3> HiResRenderDynamicMesh;
+
+		UStaticMesh* StaticMesh;
+	};
+
+
+	typedef ULODManagerTool::FLODMeshInfo  FLODMeshInfo;
+	typedef ULODManagerTool::FLODName      FLODName;
+	
+	// proxy for information that the tool mutates, i.e. a mapping of LOD index to imported LODs and list of deleted materials
+	struct FProxyLODState
+	{
+		// Identifies the imported LOD in the original static mesh.
+		struct FImportedLODId
+		{
+			FImportedLODId() : SrcLODIndex(-2){}
+
+			int32 SrcLODIndex = -2; // the index of this LOD in the source static mesh.  -1 is hi-res.
+
+			bool IsValid() const { return SrcLODIndex > -2; }
+			bool IsHighRes() const { return SrcLODIndex == -1;}
+			void Reset() 
+			{ 
+				SrcLODIndex = -2;
+			}
+
+			bool operator==(const FImportedLODId& Other) const
+			{
+				return (SrcLODIndex == Other.SrcLODIndex);
+			}
+		};
+
+
+		FProxyLODState() {}
+
+		FProxyLODState(const UStaticMesh& StaticMesh)
+		{
+			Reset(StaticMesh);
+		}
+
+		// clear data
+		void Reset()
+		{
+			HiRes.Reset();
+			LODs.Reset();
+		}
+
+		// reset the imported lod mapping and clear the list of deleted materials
+		void Reset(const UStaticMesh& StaticMesh)
+		{
+			Reset();
+
+			if (StaticMesh.IsHiResMeshDescriptionValid())
+			{
+				HiRes.SrcLODIndex = -1;
+			}
+
+			int32 NumSourceModels = StaticMesh.GetNumSourceModels();
+			LODs.SetNum(NumSourceModels);
+			for (int i = 0; i < NumSourceModels; ++i)
+			{
+				EMeshLODIdentifier LOD = (i < 8) ? static_cast<EMeshLODIdentifier>(i) : EMeshLODIdentifier::LOD7;
+				LODs[i].SrcLODIndex = i;
+			}
+		}
+
+		void DeleteHiResSourceModel()
+		{
+			HiRes.Reset();
+		}
+
+		void MoveHiResToLOD0()
+		{
+			if (HiRes.IsValid())
+			{
+				LODs[0] = HiRes;
+				HiRes.Reset();
+			}
+		}
+
+		bool HasHiResAsLOD0() const
+		{
+			return !HiRes.IsValid() && LODs.IsValidIndex(0) && LODs[0].IsHighRes();
+		}
+
+		// return true if there are no materials to delete and the imported LOD order matches the static mesh 
+		//( i.e. no changes are needed to Sync the static mesh )
+		bool IsInSync(const UStaticMesh& StaticMesh) const
+		{
+			if (DeletedMaterialIndices.IsEmpty())
+			{
+				if (StaticMesh.IsHiResMeshDescriptionValid() == HiRes.IsValid())
+				{
+					bool IsSame = true;
+					int32 NumSourceModels = StaticMesh.GetNumSourceModels();
+					for (int i = 0; i < NumSourceModels; ++i)
+					{
+						EMeshLODIdentifier LOD = (i < 8) ? static_cast<EMeshLODIdentifier>(i) : EMeshLODIdentifier::LOD7;
+						IsSame = IsSame && (LODs[i].SrcLODIndex == i);
+					}
+					return IsSame;
+				}
+			}
+			return false;
+ 		}
+
+
+		// update the static mesh to match the managed state.  
+		// Note, this calls modify on the mesh and should happen inside a transaction!
+		bool ModifyStaticMesh(UStaticMesh& StaticMesh)
+		{
+			if (IsInSync(StaticMesh)) // nothing to be done.
+			{
+				return false;
+			}
+
+			const bool bDeleteMaterials = ( DeletedMaterialIndices.Num() > 0 );
+			bool bUpdatedLOD0 = false;
+
+			TArray<FStaticMaterial> NewMaterialSet;
+			TMap<FPolygonGroupID, FPolygonGroupID> RemapPolygonGroups;
+			
+			if (bDeleteMaterials) // compute NewMaterialSet and RemapPolygonGroups
+			{
+				const TArray<FStaticMaterial>& CurrentMaterialSet = StaticMesh.GetStaticMaterials();
+				const int32 NumMaterials = CurrentMaterialSet.Num();
+				TArray<int32> MaterialIDMap;
+				MaterialIDMap.Init(-1, NumMaterials);
+				for (int32 k = 0; k < NumMaterials; ++k)
+				{
+					if (!DeletedMaterialIndices.Contains(k))
+					{
+						MaterialIDMap[k] = NewMaterialSet.Num();
+						RemapPolygonGroups.Add(FPolygonGroupID(k), FPolygonGroupID(MaterialIDMap[k]));
+						NewMaterialSet.Add(CurrentMaterialSet[k]);
+					}
+				}
+			}
+
+			StaticMesh.Modify();
+
+			// was the high res moved or deleted?
+			if (StaticMesh.IsHiResMeshDescriptionValid() && !HiRes.IsValid())
+			{
+
+				StaticMesh.ModifyHiResMeshDescription();
+
+				bool bHiResMoved = (LODs[0].IsHighRes());
+
+				if (bHiResMoved)
+				{
+					FMeshDescription* HiResMeshDescription = StaticMesh.GetHiResMeshDescription();
+
+					StaticMesh.ModifyMeshDescription(0);
+					FMeshDescription* LOD0MeshDescription = StaticMesh.GetMeshDescription(0);
+					*LOD0MeshDescription = *HiResMeshDescription;
+					if (bDeleteMaterials)
+					{
+						LOD0MeshDescription->RemapPolygonGroups(RemapPolygonGroups);
+					}
+					StaticMesh.CommitMeshDescription(0);
+
+					bUpdatedLOD0 = true;
+				}
+			
+				StaticMesh.ClearHiResMeshDescription();
+				StaticMesh.CommitHiResMeshDescription();
+				
+			}
+
+			if (bDeleteMaterials)
+			{
+				// update the material set
+				StaticMesh.GetStaticMaterials() = NewMaterialSet;
+
+				// remap the polygon groups in the retained mesh descriptions
+				{
+					if (HiRes.IsValid())
+					{
+						StaticMesh.ModifyHiResMeshDescription();
+						FMeshDescription* HiResMeshDescription = StaticMesh.GetHiResMeshDescription();
+						HiResMeshDescription->RemapPolygonGroups(RemapPolygonGroups);
+						StaticMesh.CommitHiResMeshDescription();
+					}
+
+					const int32 kStart = (bUpdatedLOD0 == false) ? 0 : 1;
+
+
+					for (int32 k = kStart; k < StaticMesh.GetNumSourceModels(); ++k)
+					{
+						if (StaticMesh.IsSourceModelValid(k))
+						{
+							if (FMeshDescription* Mesh = StaticMesh.GetMeshDescription(k))
+							{
+								StaticMesh.ModifyMeshDescription(k);
+								Mesh->RemapPolygonGroups(RemapPolygonGroups);
+								StaticMesh.CommitMeshDescription(k);
+							}
+						}
+					}
+				}
+			}
+
+			return true;
+
+
+		}
+
+		// information this tools tracks
+		FImportedLODId HiRes;
+		TArray< FImportedLODId > LODs;  // map LOD index to imported LODs.
+		TSet<int32> DeletedMaterialIndices;
+
+	};
+
+
+	/**
+	* FLODManagerToolChange represents an reversible change to ProxyLODState
+	* ULODManagerTool is the only supported target.
+	*/
+	class FLODManagerToolChange : public FToolCommandChange
+	{
+	public:
+		
+		FLODManagerToolChange(const FText& ChangeDescription)
+		:Description(ChangeDescription)
+		{}
+
+		FProxyLODState OldMeshLODState;
+		FProxyLODState NewMeshLODState;
+
+		/** Makes the change to the object */
+		virtual void Apply(UObject* Object) override 
+		{
+			ULODManagerTool* ChangeTarget = CastChecked<ULODManagerTool>(Object);
+			ChangeTarget->ApplyChange(this, false);
+		}
+
+		/** Reverts change to the object */
+		virtual void Revert(UObject* Object) override
+		{
+			ULODManagerTool* ChangeTarget = CastChecked<ULODManagerTool>(Object);
+			ChangeTarget->ApplyChange(this, true);
+		}
+
+		/** Describes this change (for debugging) */
+		virtual FString ToString() const override
+		{
+			return (Description.IsEmpty()) ? FString(TEXT("LODManager Change")) : Description.ToString();
+		}
+
+		FText Description;	
+	};
+
+}}}
 
 
 
@@ -86,11 +689,27 @@ void ULODManagerTool::Setup()
 		return;
 	}
 
+	UStaticMesh* StaticMesh = nullptr;
+	if (UStaticMeshComponent* StaticMeshComponent = Cast<UStaticMeshComponent>(UE::ToolTarget::GetTargetComponent(Targets[0])))
+	{
+		StaticMesh = StaticMeshComponent->GetStaticMesh();
+		if (StaticMesh)
+		{ 
+			UE::ToolTarget::HideSourceObject(Targets[0]);
 
+			ProxyLODState = MakeUnique<UE::Geometry::LODManagerHelper::FProxyLODState>(*StaticMesh);
+			DynamicMeshCache = MakeUnique<UE::Geometry::LODManagerHelper::FDynamicMeshLODCache>(*StaticMesh);
+		}
+	}
+
+	const bool bHasHiRes = (StaticMesh) ? StaticMesh->IsHiResMeshDescriptionValid() : false;
 	LODPreviewProperties = NewObject<ULODManagerPreviewLODProperties>(this);
 	AddToolPropertySource(LODPreviewProperties);
-	LODPreviewProperties->VisibleLOD = this->DefaultLODName;
-	LODPreviewProperties->bShowingDefaultLOD = true;
+	
+	UpdateLODNames();
+	const int32 DefaultNameIndex = (bHasHiRes) ? 1 : 0; // want the index of the first LOD name, not the hi-res
+	
+	LODPreviewProperties->VisibleLOD = LODPreviewProperties->LODNamesList[DefaultNameIndex];
 	LODPreviewProperties->WatchProperty(LODPreviewProperties->VisibleLOD, [this](FString NewLOD) { bPreviewLODValid = false; });
 	LODPreviewProperties->WatchProperty(LODPreviewProperties->bShowSeams, [this](bool bNewValue) { bPreviewLODValid = false; });
 
@@ -107,18 +726,14 @@ void ULODManagerTool::Setup()
 
 	bLODInfoValid = false;
 
-	if (UStaticMeshComponent* StaticMeshComponent = Cast<UStaticMeshComponent>(UE::ToolTarget::GetTargetComponent(Targets[0])))
+	// Optionally add buttons to delete or move the hi-res mesh
+	if (bHasHiRes)
 	{
-		if (UStaticMesh* StaticMesh = StaticMeshComponent->GetStaticMesh())
-		{
-			if (StaticMesh->IsHiResMeshDescriptionValid())
-			{
-				HiResSourceModelActions = NewObject<ULODManagerHiResSourceModelActions>(this);
-				HiResSourceModelActions->Initialize(this);
-				AddToolPropertySource(HiResSourceModelActions);
-			}
-		}
+		HiResSourceModelActions = NewObject<ULODManagerHiResSourceModelActions>(this);
+		HiResSourceModelActions->Initialize(this);
+		AddToolPropertySource(HiResSourceModelActions);
 	}
+	
 
 	MaterialActions = NewObject<ULODManagerMaterialActions>(this);
 	MaterialActions->Initialize(this);
@@ -135,7 +750,39 @@ void ULODManagerTool::Setup()
 }
 
 
+void ULODManagerTool::ApplyChange(const UE::Geometry::LODManagerHelper::FLODManagerToolChange* Change, bool bRevert)
+{
+	check(ProxyLODState)
+	if (!bRevert)
+	{
+		*ProxyLODState = Change->NewMeshLODState;
+	}
+	else
+	{
+		*ProxyLODState = Change->OldMeshLODState;
+	}
+}
+
+// copy current state into the active change as Old
+void ULODManagerTool::BeginChange(FText TransactionName)
+{
+	if(!ActiveChange.IsValid())
+	{ 
+		ActiveChange = MakeUnique< UE::Geometry::LODManagerHelper::FLODManagerToolChange >(TransactionName);
+		ActiveChange->OldMeshLODState = *ProxyLODState;
+	}
+}
 
+// copy current state into the active change as New and emit change object
+void ULODManagerTool::EndChange()
+{
+	if (ActiveChange.IsValid())
+	{
+		ActiveChange->NewMeshLODState= *ProxyLODState;
+		const FText& TransactionName = ActiveChange->Description;
+		GetToolManager()->EmitObjectChange(this, MoveTemp(ActiveChange), TransactionName);
+	}
+}
 
 void ULODManagerTool::OnShutdown(EToolShutdownType ShutdownType)
 {
@@ -144,10 +791,28 @@ void ULODManagerTool::OnShutdown(EToolShutdownType ShutdownType)
 		LODPreview->Disconnect();
 		LODPreviewLines->Disconnect();
 	}
+	// restore visibility
 	for (int32 ComponentIdx = 0; ComponentIdx < Targets.Num(); ComponentIdx++)
 	{
 		UE::ToolTarget::ShowSourceObject(Targets[ComponentIdx]);
 	}
+
+	if (UStaticMesh* StaticMesh = GetSingleStaticMesh())
+	{ 
+		if (ShutdownType == EToolShutdownType::Accept)
+		{ 
+			// the regular transaction 
+			GetToolManager()->BeginUndoTransaction(LOCTEXT("AcceptLODManager", "Accept LOD Manager"));
+			bool bModifiedMesh = ProxyLODState->ModifyStaticMesh(*StaticMesh);
+		
+			// @todo consider moving PostEditChange to outside the scope of the transaction. 
+			if (bModifiedMesh)
+			{
+				StaticMesh->PostEditChange();
+			}
+			GetToolManager()->EndUndoTransaction();
+		}
+	}
 }
 
 
@@ -181,24 +846,28 @@ void ULODManagerTool::OnTick(float DeltaTime)
 	PendingAction = ELODManagerToolActions::NoAction;
 
 	// test capture a dirty state that results from undo while in the tool.
-	auto IsLODInfoDirty = [&]()
+	bool IsLODInfoDirty = [&, this]
 		{
 			UStaticMesh* StaticMesh = GetSingleStaticMesh();
-			if (!StaticMesh || !LODInfoProperties)
+			if (!StaticMesh || !LODInfoProperties || !ProxyLODState)
 			{
 				return false;
 			}
 
 			// have we changed the number of materials
-			bool bDirty = (LODInfoProperties->Materials.Num() != StaticMesh->GetStaticMaterials().Num());
+			const int32 NumDeletedMaterials =  ProxyLODState->DeletedMaterialIndices.Num();
+			const int32 CurrentNumMaterial = StaticMesh->GetStaticMaterials().Num() - NumDeletedMaterials;
+			const int32 DisplayedNumMaterials = LODInfoProperties->Materials.Num();
+			
+			bool bDirty = (DisplayedNumMaterials != CurrentNumMaterial);
 
 			// what about hi-res mesh?
-			bDirty = bDirty || (StaticMesh->IsHiResMeshDescriptionValid() ? LODInfoProperties->HiResSource.Num() == 0 : LODInfoProperties->HiResSource.Num() != 0);
+			bDirty = bDirty || (ProxyLODState->HiRes.IsValid() ? LODInfoProperties->HiResSource.Num() == 0 : LODInfoProperties->HiResSource.Num() != 0);
 
 			return bDirty;
-		};
+		}();
 
-	bLODInfoValid = bLODInfoValid && !IsLODInfoDirty();
+	bLODInfoValid = bLODInfoValid && !IsLODInfoDirty;
 
 	if (bLODInfoValid == false)
 	{
@@ -220,19 +889,7 @@ UStaticMesh* ULODManagerTool::GetSingleStaticMesh()
 		return nullptr;
 	}
 	UStaticMeshComponent* StaticMeshComponent = Cast<UStaticMeshComponent>(UE::ToolTarget::GetTargetComponent(Targets[0]));
-	if (StaticMeshComponent == nullptr)
-	{
-		// We may have removed the highest quality LOD ( i.e hi-res ) during tool operation or restored it by doing an in-tool undo,
-		// so here we reset the target to current max quality.  Why? Because if the tool target 'Editing LOD' does not correspond to an existing LOD on 
-		// the mesh then target is considered invalid  and GetTargetComponent will fail as will showing the source object on shutdown.
-		if (UStaticMeshComponentToolTarget* StaticMeshComponentTarget = Cast<UStaticMeshComponentToolTarget>(Targets[0]))
-		{
-			StaticMeshComponentTarget->SetEditingLOD(EMeshLODIdentifier::MaxQuality);
-
-			StaticMeshComponent = Cast<UStaticMeshComponent>(UE::ToolTarget::GetTargetComponent(Targets[0]));
-		}	
-	}
-
+	
 	if (StaticMeshComponent == nullptr || StaticMeshComponent->GetStaticMesh() == nullptr)
 	{
 		return nullptr;
@@ -249,6 +906,8 @@ UStaticMesh* ULODManagerTool::GetSingleStaticMesh()
 
 void ULODManagerTool::UpdateLODInfo()
 {
+	using namespace UE::Geometry::LODManagerHelper;
+
 	UStaticMesh* StaticMesh = GetSingleStaticMesh();
 	if (bLODInfoValid == true || StaticMesh == nullptr)
 	{
@@ -256,50 +915,68 @@ void ULODManagerTool::UpdateLODInfo()
 	}
 	bLODInfoValid = true;
 
+	// update the displayed info about the LODS
 	LODInfoProperties->bNaniteEnabled = StaticMesh->IsNaniteEnabled();
 	LODInfoProperties->KeepTrianglePercent = StaticMesh->NaniteSettings.KeepPercentTriangles * 100;
 
-	TArray<FStaticMaterial> CurMaterialSet = StaticMesh->GetStaticMaterials();
-	LODInfoProperties->Materials.Reset();
-	for (const FStaticMaterial& Material : StaticMesh->GetStaticMaterials())
+	const TArray<FStaticMaterial>& OriginalMaterialSet = StaticMesh->GetStaticMaterials();
+	
+	// filter the original material set, against those pending delete, into the displayed materials (LODInfoProperties)
+	LODInfoProperties->Materials.Reset();	
+	for (int32 mid = 0, NumMaterials =  OriginalMaterialSet.Num(); mid < NumMaterials; ++mid )
 	{
-		LODInfoProperties->Materials.Add(Material);
+		if (!ProxyLODState->DeletedMaterialIndices.Contains(mid))
+		{
+			LODInfoProperties->Materials.Add(OriginalMaterialSet[mid]);
+		}
 	}
 
 
 	// per-LOD info
 	LODInfoProperties->SourceLODs.Reset();
-	int32 NumSourceModels = StaticMesh->GetNumSourceModels();
-	for (int32 si = 0; si < NumSourceModels; ++si)
+	for (const FProxyLODState::FImportedLODId& ImportedLODInfo : ProxyLODState->LODs)
 	{
-		const FStaticMeshSourceModel& LODSourceModel = StaticMesh->GetSourceModel(si);
-		const FMeshDescription* LODMesh = StaticMesh->GetMeshDescription(si);
-		if (LODMesh != nullptr)		// generated LODs do not have mesh description...is that the only way to tell?
+		if (ImportedLODInfo.IsValid())
 		{
-			FLODManagerLODInfo LODInfo = { LODMesh->Vertices().Num(), LODMesh->Triangles().Num() };
+			const FDynamicMeshLODCache::FMeshElementInfo ElementCounts = (ImportedLODInfo.IsHighRes()) ? DynamicMeshCache->GetHiResMeshInfo() : DynamicMeshCache->GetLODMeshElementInfo(ImportedLODInfo.SrcLODIndex);
+			const FLODManagerLODInfo LODInfo = { ElementCounts.VertexCount, ElementCounts.TriangleCount };
 			LODInfoProperties->SourceLODs.Add(LODInfo);
 		}
 	}
 
+	// hi-res source info
 	LODInfoProperties->HiResSource.Reset();
-	if (StaticMesh->IsHiResMeshDescriptionValid())
+	if (ProxyLODState->HiRes.IsValid() && ProxyLODState->HiRes.IsHighRes())
 	{
-		const FMeshDescription* HiResMesh = StaticMesh->GetHiResMeshDescription();
-		FLODManagerLODInfo LODInfo = { HiResMesh->Vertices().Num(), HiResMesh->Triangles().Num() };
+		const FDynamicMeshLODCache::FMeshElementInfo& ElementCounts = DynamicMeshCache->GetHiResMeshInfo();
+		const FLODManagerLODInfo LODInfo = { ElementCounts.VertexCount, ElementCounts.TriangleCount };
 		LODInfoProperties->HiResSource.Add(LODInfo);
 	}
 
+	// render data lod info
 	LODInfoProperties->RenderLODs.Reset();
+	const bool bHiResInLOD0 = ProxyLODState->LODs.IsValidIndex(0) && ProxyLODState->LODs[0].IsHighRes();
 	if (StaticMesh->HasValidRenderData())
 	{
 		const FStaticMeshRenderData* RenderData = StaticMesh->GetRenderData();
-		for (const FStaticMeshLODResources& LODResource : RenderData->LODResources)
+		const int32 NumRenderLODs = RenderData->LODResources.Num();
+
+		if (bHiResInLOD0)
 		{
-			FLODManagerLODInfo LODInfo;
-			LODInfo.VertexCount = LODResource.VertexBuffers.PositionVertexBuffer.GetNumVertices();
-			LODInfo.TriangleCount = LODResource.IndexBuffer.GetNumIndices() / 3;
+			const FDynamicMeshLODCache::FMeshElementInfo ElementCounts = DynamicMeshCache->GetHiResLODAsRenderDataMeshInfo();
+			FLODManagerLODInfo LODInfo = {ElementCounts.VertexCount, ElementCounts.TriangleCount};
 			LODInfoProperties->RenderLODs.Add(LODInfo);
 		}
+		
+
+		for (int32 i = (bHiResInLOD0)? 1: 0 ; i < NumRenderLODs; ++i)
+		{ 
+			if (DynamicMeshCache->HasRenderLOD(i))
+			{
+				const FDynamicMeshLODCache::FMeshElementInfo ElementCounts = DynamicMeshCache->GetRenderLODMeshInfo(i);
+				FLODManagerLODInfo LODInfo = { ElementCounts.VertexCount, ElementCounts.TriangleCount };
+			}
+		}
 	}
 
 
@@ -310,7 +987,6 @@ void ULODManagerTool::UpdateLODInfo()
 
 
 
-
 void ULODManagerTool::UpdateLODNames()
 {
 	UStaticMesh* StaticMesh = GetSingleStaticMesh();
@@ -323,42 +999,63 @@ void ULODManagerTool::UpdateLODNames()
 	UILODNames.Reset();
 	ActiveLODNames.Reset();
 
-	UILODNames.Add(DefaultLODName);
-	ActiveLODNames.Add(DefaultLODName, FLODName());
-
-	if (StaticMesh->IsHiResMeshDescriptionValid())
+	// Hi-Res 
+	if (ProxyLODState->HiRes.IsValid())
 	{
 		FString LODName(TEXT("SourceModel HiRes"));
 		UILODNames.Add(LODName);
 		ActiveLODNames.Add(LODName, { -1, -1, 0 });
 	}
 
-	int32 NumSourceModels = StaticMesh->GetNumSourceModels();
+	// Standard LODs
+	int32 NumSourceModels = ProxyLODState->LODs.Num();
 	for (int32 si = 0; si < NumSourceModels; ++si)
-	{
-		const FStaticMeshSourceModel& LODSourceModel = StaticMesh->GetSourceModel(si);
-		const FMeshDescription* LODMesh = StaticMesh->GetMeshDescription(si);
-		if (LODMesh != nullptr)		// generated LODs do not have mesh description...is that the only way to tell?
+	{ 
+		const UE::Geometry::LODManagerHelper::FProxyLODState::FImportedLODId& LODInfo = ProxyLODState->LODs[si];
+	
+		if (LODInfo.IsValid())
 		{
-			FString LODName = FString::Printf(TEXT("SourceModel LOD%d"), si);
-			UILODNames.Add(LODName);
-			ActiveLODNames.Add(LODName, { si, -1, -1 });
+			FString LODUIName = FString::Printf(TEXT("SourceModel LOD%d"), si);
+			UILODNames.Add(LODUIName);
+
+			if (LODInfo.IsHighRes()) // test if this is a moved hi-res
+			{
+				ActiveLODNames.Add(LODUIName, { -1, -1, 0 });
+			}
+			else
+			{
+				ActiveLODNames.Add(LODUIName, { LODInfo.SrcLODIndex, -1, -1 });
+			}
 		}
 	}
 
+	const bool bHiResAtLOD0 = ProxyLODState->HasHiResAsLOD0();
+
 	if (StaticMesh->HasValidRenderData())
 	{
 		const FStaticMeshRenderData* RenderData = StaticMesh->GetRenderData();
 		int32 Index = 0;
 		for (const FStaticMeshLODResources& LODResource : RenderData->LODResources)
 		{
-			FString LODName = FString::Printf(TEXT("RenderData LOD%d"), Index);
-			UILODNames.Add(LODName);
-			ActiveLODNames.Add(LODName, { -1, Index, -1 });
+			FString LODUIName = FString::Printf(TEXT("RenderData LOD%d"), Index);
+			UILODNames.Add(LODUIName);
+			if (Index == 0 && bHiResAtLOD0) // test if we need render data for moved hi-res at lod 0
+			{
+				ActiveLODNames.Add(LODUIName, { -1, -1, 1 }); // render data version of original hi-res
+			}
+			else
+			{
+				ActiveLODNames.Add(LODUIName, { -1, Index, -1 });
+			}
 			Index++;
 		}
 	}
 
+	// update the name selected in the UI if needed.
+	if (!ActiveLODNames.Contains(LODPreviewProperties->VisibleLOD))
+	{
+		LODPreviewProperties->VisibleLOD = UILODNames[0];
+	}
 }
 
 
@@ -367,18 +1064,7 @@ void ULODManagerTool::UpdateLODNames()
 
 void ULODManagerTool::UpdatePreviewLines(FLODMeshInfo& MeshInfo)
 {
-	if (MeshInfo.bInfoCached == false)
-	{
-		MeshInfo.bInfoCached = true;
-		for (int32 eid : MeshInfo.Mesh.EdgeIndicesItr())
-		{
-			if (MeshInfo.Mesh.IsBoundaryEdge(eid))
-			{
-				MeshInfo.BoundaryEdges.Add(eid);
-			}
-		}
-	}
-
+	
 	float LineWidthMultiplier = 1.0f;
 	FColor BoundaryEdgeColor(240, 15, 15);
 	float BoundaryEdgeThickness = LineWidthMultiplier * 4.0f;
@@ -395,7 +1081,7 @@ void ULODManagerTool::UpdatePreviewLines(FLODMeshInfo& MeshInfo)
 		for (int32 eid : MeshInfo.BoundaryEdges)
 		{
 			FVector3d A, B;
-			MeshInfo.Mesh.GetEdgeV(eid, A, B);
+			MeshInfo.Mesh->GetEdgeV(eid, A, B);
 			BoundaryLines->AddLine((FVector)A, (FVector)B, BoundaryEdgeColor, BoundaryEdgeThickness, BoundaryEdgeDepthBias);
 		}
 	}
@@ -421,364 +1107,153 @@ void ULODManagerTool::UpdatePreviewLOD()
 	}
 	bPreviewLODValid = true;
 
-	auto SetShowingDefaultLOD = [this](bool bNewValue)
-	{
-		if (bNewValue != LODPreviewProperties->bShowingDefaultLOD)
-		{
-			LODPreviewProperties->bShowingDefaultLOD = bNewValue;
-			if (bNewValue == true)
-			{
-				LODPreviewProperties->VisibleLOD = DefaultLODName;
-			}
-			NotifyOfPropertyChangeByTool(LODPreviewProperties);
-		}
-	};
 
-	FString SelectedLOD = LODPreviewProperties->VisibleLOD;
-	const FLODName* FoundName = ActiveLODNames.Find(SelectedLOD);
-	if (SelectedLOD.IsEmpty() || FoundName == nullptr || FoundName->IsDefault() )
-	{
-		// currently just showing the source object when Default is selected
-		LODPreview->SetVisible(false);
-		ClearPreviewLines();
-		LODPreviewLines->SetAllVisible(false);
-		UE::ToolTarget::ShowSourceObject(Targets[0]);
-		SetShowingDefaultLOD(true);
-		return;
-	}
+	// the UI mesh selection
+	const FString SelectedLOD = LODPreviewProperties->VisibleLOD;
 
-	TUniquePtr<FLODMeshInfo>* Found = LODMeshCache.Find(SelectedLOD);
-	if (Found == nullptr)
+	// Get an identifier for the src mesh
+	if (const FLODName* FoundName = ActiveLODNames.Find(SelectedLOD))
 	{
-		CacheLODMesh(SelectedLOD, *FoundName);
-		Found = LODMeshCache.Find(SelectedLOD);
+		TUniquePtr<FLODMeshInfo> MeshInfo = GetLODMeshInfo(*FoundName);
+		
+		if (MeshInfo->Mesh.IsValid())
+		{ 
+			LODPreview->ReplaceMesh(*MeshInfo->Mesh);
+			LODPreview->SetVisible(true);
+			UpdatePreviewLines(*MeshInfo);
+			LODPreviewLines->SetAllVisible(LODPreviewProperties->bShowSeams);
+		}
 	}
-
-	SetShowingDefaultLOD(Found == nullptr);
 	
-	if (Found)
-	{
-		UE::ToolTarget::HideSourceObject(Targets[0]);
-
-		LODPreview->ReplaceMesh((*Found)->Mesh);
-		LODPreview->SetVisible(true);
-		UpdatePreviewLines(**Found);
-		LODPreviewLines->SetAllVisible(LODPreviewProperties->bShowSeams);
-	}
-	else
-	{
-		LODPreview->SetVisible(false);
-		ClearPreviewLines();
-		LODPreviewLines->SetAllVisible(false);
-		UE::ToolTarget::ShowSourceObject(Targets[0]);
-	}
 }
 
 
 
 
-bool ULODManagerTool::CacheLODMesh(const FString& Name, FLODName LODName)
+TUniquePtr<ULODManagerTool::FLODMeshInfo> ULODManagerTool::GetLODMeshInfo(const FLODName& LODName)
 {
-	UStaticMesh* StaticMesh = GetSingleStaticMesh();
-	if (!ensure(StaticMesh)) return false;
-
-	FMeshDescriptionToDynamicMesh Converter;
+	TSharedPtr<FDynamicMesh3> AsDynamicMesh;
 
-	if (LODName.OtherIndex == 0)		// HiRes SourceModel
+	if (LODName.OtherIndex == 0) // HiResSourceModel
 	{
-		FMeshDescription* HiResMeshDescription = StaticMesh->GetHiResMeshDescription();
-		if (HiResMeshDescription != nullptr)
-		{
-			// TODO: we only need to do this copy if normals or tangents are set to auto-generated...
-			FMeshBuildSettings HiResBuildSettings = StaticMesh->GetHiResSourceModel().BuildSettings;
-			FMeshDescription TmpMeshDescription(*HiResMeshDescription);
-			UE::MeshDescription::InitializeAutoGeneratedAttributes(TmpMeshDescription, &HiResBuildSettings);
-
-			TUniquePtr<FLODMeshInfo> NewMeshInfo = MakeUnique<FLODMeshInfo>();
-			Converter.Convert(&TmpMeshDescription, NewMeshInfo->Mesh, true);
-			LODMeshCache.Add(Name, MoveTemp(NewMeshInfo));
-			return true;
-		}
-
+		AsDynamicMesh = DynamicMeshCache->GetHiResLOD();
 	}
-	else if (LODName.SourceModelIndex >= 0)
+	else if (LODName.OtherIndex == 1) // Artificial version of the renderdata for the hi-res
 	{
-		FMeshDescription* LODMeshDescription = StaticMesh->GetMeshDescription(LODName.SourceModelIndex);
-		if (LODMeshDescription != nullptr)
-		{
-			// TODO: we only need to do this copy if normals or tangents are set to auto-generated...
-			FMeshBuildSettings LODBuildSettings = StaticMesh->GetSourceModel(LODName.SourceModelIndex).BuildSettings;
-			FMeshDescription TmpMeshDescription(*LODMeshDescription);
-			UE::MeshDescription::InitializeAutoGeneratedAttributes(TmpMeshDescription, &LODBuildSettings);
-
-			TUniquePtr<FLODMeshInfo> NewMeshInfo = MakeUnique<FLODMeshInfo>();
-			Converter.Convert(&TmpMeshDescription, NewMeshInfo->Mesh, true);
-			LODMeshCache.Add(Name, MoveTemp(NewMeshInfo));
-			return true;
-		}
+		AsDynamicMesh = DynamicMeshCache->GetHiResLODAsRenderData();
 	}
-	else if (LODName.RenderDataIndex >= 0)
+	else if (LODName.SourceModelIndex >= 0) // requesting a source model
 	{
-		if (StaticMesh->HasValidRenderData())
+		AsDynamicMesh = DynamicMeshCache->GetImportedLOD(LODName.SourceModelIndex);
+	}
+	else if (LODName.RenderDataIndex >= 0) // requesting render data
+	{
+		AsDynamicMesh = DynamicMeshCache->GetRenderLOD(LODName.RenderDataIndex);
+	}
+
+	TUniquePtr<FLODMeshInfo> NewMeshInfo = MakeUnique<FLODMeshInfo>();
+	if (AsDynamicMesh)
+	{	
+		NewMeshInfo->Mesh = AsDynamicMesh; 
+		for (int32 eid : NewMeshInfo->Mesh->EdgeIndicesItr())
 		{
-			const FStaticMeshRenderData* RenderData = StaticMesh->GetRenderData();
-			if (LODName.RenderDataIndex < RenderData->LODResources.Num())
+			if (NewMeshInfo->Mesh->IsBoundaryEdge(eid))
 			{
-				const FStaticMeshLODResources& LODResource = RenderData->LODResources[LODName.RenderDataIndex];
-				TUniquePtr<FLODMeshInfo> NewMeshInfo = MakeUnique<FLODMeshInfo>();
-				UE::Conversion::RenderBuffersToDynamicMesh(LODResource.VertexBuffers, LODResource.IndexBuffer, LODResource.Sections, NewMeshInfo->Mesh);
-				LODMeshCache.Add(Name, MoveTemp(NewMeshInfo));
-				return true;
+				NewMeshInfo->BoundaryEdges.Add(eid);
 			}
 		}
 	}
-
-	return false;
+	return MoveTemp(NewMeshInfo);
 }
 
 
-
 void ULODManagerTool::DeleteHiResSourceModel()
 {
-	GetToolManager()->BeginUndoTransaction(LOCTEXT("DeleteHiResSourceModel", "Delete HiRes Source"));
-
-	const bool bComponentsHiddenByTool = LODPreview->IsVisible();
-	for (int32 ComponentIdx = 0; ComponentIdx < Targets.Num(); ComponentIdx++)
-	{
-		UStaticMeshComponent* StaticMeshComponent = Cast<UStaticMeshComponent>(UE::ToolTarget::GetTargetComponent(Targets[ComponentIdx]));
-		if (StaticMeshComponent == nullptr || StaticMeshComponent->GetStaticMesh() == nullptr)
-		{
-			continue;
-		}
-		UStaticMesh* StaticMesh = StaticMeshComponent->GetStaticMesh();
-
-		if (StaticMesh->IsHiResMeshDescriptionValid())
-		{
-			
-			if (bComponentsHiddenByTool)
-			{ 
-				// temporarily set the visible flag - that way an undo after the tool is closed will restore a visible mesh
-				UE::ToolTarget::ShowSourceObject(Targets[ComponentIdx]);
-			}
-
-			StaticMesh->Modify();
-
-			StaticMesh->ModifyHiResMeshDescription();
-			StaticMesh->ClearHiResMeshDescription();
-			StaticMesh->CommitHiResMeshDescription();
-
-			StaticMesh->PostEditChange();
+	BeginChange(LOCTEXT("LODManagerDeleteHiRes", "Delete HiRes"));
+	ProxyLODState->DeleteHiResSourceModel();
+	EndChange();
 
-			// restore the visibility state expected by the tool
-			if (bComponentsHiddenByTool)
-			{
-				UE::ToolTarget::HideSourceObject(Targets[ComponentIdx]);
-			}
-		}
-	}
-
-	GetToolManager()->EndUndoTransaction();
-
-	LODMeshCache.Reset();
 	bLODInfoValid = false;
 }
 
 
-
 void ULODManagerTool::MoveHiResToLOD0()
 {
-	GetToolManager()->BeginUndoTransaction(LOCTEXT("RestoreLOD0", "Move HiRes to LOD0"));
-
-	const bool bComponentsHiddenByTool = LODPreview->IsVisible();
-	for (int32 ComponentIdx = 0; ComponentIdx < Targets.Num(); ComponentIdx++)
-	{
-		UStaticMeshComponent* StaticMeshComponent = Cast<UStaticMeshComponent>(UE::ToolTarget::GetTargetComponent(Targets[ComponentIdx]));
-		if (StaticMeshComponent == nullptr || StaticMeshComponent->GetStaticMesh() == nullptr)
-		{
-			continue;
-		}
-		UStaticMesh* StaticMesh = StaticMeshComponent->GetStaticMesh();
-
-		if (StaticMesh->IsHiResMeshDescriptionValid())
-		{
+	BeginChange(LOCTEXT("LODManagerMoveHiRes", "Move HiRes to LOD0"));
+	ProxyLODState->MoveHiResToLOD0();
+	EndChange();
 
-			if (bComponentsHiddenByTool)
-			{
-				// temporarily set the visible flag - that way an undo after the tool is closed will restore a visible mesh
-				UE::ToolTarget::ShowSourceObject(Targets[ComponentIdx]);
-			}
-
-			StaticMesh->Modify();
-
-			StaticMesh->ModifyHiResMeshDescription();
-			FMeshDescription* HiResMeshDescription = StaticMesh->GetHiResMeshDescription();
-
-			StaticMesh->ModifyMeshDescription(0);
-			FMeshDescription* LOD0MeshDescription = StaticMesh->GetMeshDescription(0);
-			*LOD0MeshDescription = *HiResMeshDescription;
-
-			StaticMesh->ClearHiResMeshDescription();
-			StaticMesh->CommitHiResMeshDescription();
-
-			StaticMesh->CommitMeshDescription(0);
-
-			StaticMesh->PostEditChange();
-
-			// restore the visibility state expected by the tool
-			if (bComponentsHiddenByTool)
-			{
-				UE::ToolTarget::HideSourceObject(Targets[ComponentIdx]);
-			}
-		}
-	}
-
-	GetToolManager()->EndUndoTransaction();
-
-	LODMeshCache.Reset();
 	bLODInfoValid = false;
 }
 
 
-
-
 void ULODManagerTool::RemoveUnreferencedMaterials()
 {
-	const bool bComponentsHiddenByTool = LODPreview->IsVisible();
-
-	// temporarily set the visible flag - that way an undo after the tool is closed will restore a visible mesh
-	if (bComponentsHiddenByTool)
-	{
-		for (int32 ComponentIdx = 0; ComponentIdx < Targets.Num(); ComponentIdx++)
-		{
-			UE::ToolTarget::ShowSourceObject(Targets[ComponentIdx]);
-		}
-	}
-
-	GetToolManager()->BeginUndoTransaction(LOCTEXT("RemoveUnreferencedMaterials", "Remove Unreferenced Materials"));
-	for (int32 ComponentIdx = 0; ComponentIdx < Targets.Num(); ComponentIdx++)
-	{
-		UStaticMeshComponent* StaticMeshComponent = Cast<UStaticMeshComponent>(UE::ToolTarget::GetTargetComponent(Targets[ComponentIdx]));
-		if (StaticMeshComponent == nullptr || StaticMeshComponent->GetStaticMesh() == nullptr)
-		{
-			continue;
-		}
-		UStaticMesh* StaticMesh = StaticMeshComponent->GetStaticMesh();
+	
+	const UStaticMesh* StaticMesh = GetSingleStaticMesh();
 
-		TArray<FStaticMaterial> CurMaterialSet = StaticMesh->GetStaticMaterials();
-		int32 NumMaterials = CurMaterialSet.Num();
-		TArray<bool> MatUsedFlags;
-		MatUsedFlags.Init(false, NumMaterials);
+	// get the number of materials on the source static mesh.
+	const int32 NumMaterials = GetSingleStaticMesh()->GetStaticMaterials().Num();
 
-		TArray<const FMeshDescription*> Meshes;
-		for (int32 k = 0; k < StaticMesh->GetNumSourceModels(); ++k)
+	// compute mask, indicates if a given index in original material set corresponds material that is acually being used.
+	const TArray<bool> MatUsedFlags = [this, NumMaterials, StaticMesh]
 		{
-			if (StaticMesh->IsSourceModelValid(k))
-			{
-				Meshes.Add(StaticMesh->GetMeshDescription(k));
-			}
-		}
-		if (StaticMesh->IsHiResMeshDescriptionValid())
-		{
-			Meshes.Add(StaticMesh->GetHiResMeshDescription());
-		}
-		int32 NumMeshes = Meshes.Num();
-
-		// using RenderData instead of Static Mesh Attributes because not every Source Model will have
-		// a Mesh Description to inspect, which caused a crash
-		FStaticMeshRenderData* RenderData = StaticMesh->GetRenderData();
-		
-		// for each mesh, collect list of material indices, and then set MatUsedFlags
-		TArray<TArray<int32>> MeshMaterialIndices;
-		MeshMaterialIndices.SetNum(NumMeshes);
-		
-		for (int32 lod = 0; lod < NumMeshes; ++lod)
-		{
-			if (RenderData && RenderData->LODResources.IsValidIndex(lod))
+			TArray<bool> Used;
+			Used.Init(false, NumMaterials);
+			
+			// review materials used on the mesh in the hi-res slot.
+			if (ProxyLODState->HiRes.IsValid())
 			{
-				FStaticMeshLODResources& LOD = RenderData->LODResources[lod];
-				int NumSections = LOD.Sections.Num();
+				const int32 SrcLODIndex = ProxyLODState->HiRes.SrcLODIndex;
 
-				for (int32 SectionIndex = 0; SectionIndex < NumSections; SectionIndex++)
+				EMeshLODIdentifier LODIdentifier = (SrcLODIndex == -1) ? EMeshLODIdentifier::HiResSource : static_cast<EMeshLODIdentifier>(SrcLODIndex);
+				const TArray<int32> SectionToMaterialIndex = UStaticMeshToolTarget::MapSectionToMaterialID(StaticMesh, LODIdentifier);
+				for (int32 MaterialIndex : SectionToMaterialIndex)
 				{
-					// For every section in the LOD, retrieves the index of the material used
-					FMeshSectionInfo Info = StaticMesh->GetSectionInfoMap().Get(lod, SectionIndex);
-					int32 MaterialIndex = Info.MaterialIndex;
-					if (StaticMesh->GetStaticMaterials().IsValidIndex(MaterialIndex))
+					if (Used.IsValidIndex(MaterialIndex))
 					{
-						// material is used, set array value at its index to true
-						MatUsedFlags[MaterialIndex] = true;
-						MeshMaterialIndices[lod].AddUnique(MaterialIndex);
+						Used[MaterialIndex] = true;
 					}
 				}
 			}
-		}
 
-		// construct new material set and ID remap
-		TArray<FStaticMaterial> NewMaterialSet;
-		TArray<int32> MaterialIDMap;
-		TArray<int32> RemovedMaterials;
-		TMap<FPolygonGroupID, FPolygonGroupID> RemapPolygonGroups;
-		MaterialIDMap.Init(-1, NumMaterials);
-		for (int32 k = 0; k < NumMaterials; ++k)
-		{
-			if (MatUsedFlags[k] == true )
+			// review materials in on the meshes in the regular lod array.
+			for (const auto& LOD : ProxyLODState->LODs)
 			{
-				MaterialIDMap[k] = NewMaterialSet.Num();
-				RemapPolygonGroups.Add(FPolygonGroupID(k), FPolygonGroupID(MaterialIDMap[k]));
-				NewMaterialSet.Add(CurMaterialSet[k]);
-			}
-			else
-			{
-				RemovedMaterials.Add(k);
-			}
-		}
-		if (NewMaterialSet.Num() == CurMaterialSet.Num())
-		{
-			continue;
-		}
-
-		StaticMesh->Modify();
-
-		// update material set
-		StaticMesh->GetStaticMaterials() = NewMaterialSet;
+				const int32 SrcLODIndex = LOD.SrcLODIndex;
 
-		// now remap groups in each mesh
-		int32 MeshIndex = 0;
-		for (int32 k = 0; k < StaticMesh->GetNumSourceModels(); ++k)
-		{
-			if (StaticMesh->IsSourceModelValid(k))
-			{
-				FMeshDescription* Mesh = StaticMesh->GetMeshDescription(k);
-				if (Mesh)
+				if (LOD.IsValid())
 				{
-					StaticMesh->ModifyMeshDescription(k);
-					Mesh->RemapPolygonGroups(RemapPolygonGroups);
-					StaticMesh->CommitMeshDescription(k);
+					EMeshLODIdentifier LODIdentifier = (SrcLODIndex == -1) ? EMeshLODIdentifier::HiResSource : static_cast<EMeshLODIdentifier>(SrcLODIndex);
+					const TArray<int32> SectionToMaterialIndex = UStaticMeshToolTarget::MapSectionToMaterialID(StaticMesh, LODIdentifier);
+					for (int32 MaterialIndex : SectionToMaterialIndex)
+					{
+						if (Used.IsValidIndex(MaterialIndex))
+						{
+							Used[MaterialIndex] = true;
+						}
+					}
 				}
 			}
-		}
-		if (StaticMesh->IsHiResMeshDescriptionValid())
-		{
-			StaticMesh->ModifyHiResMeshDescription();
-			FMeshDescription* HiResMeshDescription = StaticMesh->GetHiResMeshDescription();
-			HiResMeshDescription->RemapPolygonGroups(RemapPolygonGroups);
-			StaticMesh->CommitHiResMeshDescription();
-		}
+			return Used;
+		}();
+
 
-		StaticMesh->PostEditChange();
-	}
 
-	GetToolManager()->EndUndoTransaction();
+	BeginChange(LOCTEXT("LODManagerRemoveMaterials", "Delete Unused Materials"));
 
-	// revert to the original visibility
-	if (bComponentsHiddenByTool)
+	// save list of unused (i.e. to be deleted) material indices 
+	ProxyLODState->DeletedMaterialIndices.Reset();
+	for (int32 i = 0; i < NumMaterials; ++i)
 	{
-		for (int32 ComponentIdx = 0; ComponentIdx < Targets.Num(); ComponentIdx++)
+		if (MatUsedFlags[i] == false)
 		{
-			UE::ToolTarget::HideSourceObject(Targets[ComponentIdx]);
+			ProxyLODState->DeletedMaterialIndices.Add(i);
 		}
 	}
 
+	EndChange();
+
 	bLODInfoValid = false;
 }
 
diff --git a/Engine/Plugins/Editor/MeshLODToolset/Source/MeshLODToolset/Public/Tools/LODManagerTool.h b/Engine/Plugins/Editor/MeshLODToolset/Source/MeshLODToolset/Public/Tools/LODManagerTool.h
--- a/Engine/Plugins/Editor/MeshLODToolset/Source/MeshLODToolset/Public/Tools/LODManagerTool.h
+++ b/Engine/Plugins/Editor/MeshLODToolset/Source/MeshLODToolset/Public/Tools/LODManagerTool.h
@@ -90,12 +90,9 @@ public:
 	TArray<FString> LODNamesList;
 
 	/** Control whether mesh borders are displayed */
-	UPROPERTY(EditAnywhere, DisplayName = "Show Borders", Category = LODPreview, meta = (EditCondition = "!bShowingDefaultLOD", HideEditConditionToggle))
+	UPROPERTY(EditAnywhere, DisplayName = "Show Borders", Category = LODPreview)
 	bool bShowSeams = true;
 
-	// Used to expose the VisibleLOD state to the above EditCondition
-	UPROPERTY(Transient, meta = (TransientToolProperty))
-	bool bShowingDefaultLOD = true;
 };
 
 
@@ -164,15 +161,35 @@ public:
 	}
 };
 
+// forward
+namespace UE { namespace Geometry { namespace LODManagerHelper {
+class FLODManagerToolChange; 
+struct FProxyLODState;
+class FDynamicMeshLODCache;
 
+} } }
 
 
+UINTERFACE(MinimalAPI)
+class ULODManagerToolChangeTarget : public UInterface
+{
+	GENERATED_BODY()
+};
+/**
+ * IMeshVertexCommandChangeTarget is an interface which is used to apply a FLODManagerToolChange
+ */
+class ILODManagerToolChangeTarget
+{
+	GENERATED_BODY()
+public:
+	virtual void ApplyChange(const UE::Geometry::LODManagerHelper::FLODManagerToolChange* Change, bool bRevert) = 0;
+};
 
 /**
  * Mesh Attribute Editor Tool
  */
 UCLASS()
-class MESHLODTOOLSET_API ULODManagerTool : public UMultiSelectionMeshEditingTool
+class MESHLODTOOLSET_API ULODManagerTool : public UMultiSelectionMeshEditingTool, public ILODManagerToolChangeTarget
 {
 	GENERATED_BODY()
 
@@ -183,12 +200,13 @@ public:
 	virtual void OnShutdown(EToolShutdownType ShutdownType) override;
 	virtual void OnTick(float DeltaTime) override;
 
-	virtual bool HasCancel() const override { return false; }
-	virtual bool HasAccept() const override { return false; }
-	virtual bool CanAccept() const override { return false; }
+	virtual bool HasCancel() const override { return true; }
+	virtual bool HasAccept() const override { return true; }
 
 	virtual void RequestAction(ELODManagerToolActions ActionType);
 
+	virtual void ApplyChange(const UE::Geometry::LODManagerHelper::FLODManagerToolChange* Change, bool bRevert) override;
+
 protected:
 
 	UPROPERTY()
@@ -215,34 +233,35 @@ public:
 	UFUNCTION()
 	void RemoveUnreferencedMaterials();
 
-protected:
-	UStaticMesh* GetSingleStaticMesh();
-
-	ELODManagerToolActions PendingAction = ELODManagerToolActions::NoAction;
-	
-	bool bLODInfoValid = false;
-	void UpdateLODInfo();
 
 	struct FLODName
 	{
 		int32 SourceModelIndex = -1;
 		int32 RenderDataIndex = -1;
 		int32 OtherIndex = -1;
-		bool IsDefault() const { return SourceModelIndex == -1 && RenderDataIndex == -1 && OtherIndex == -1; }
+		bool IsValid() const { return ! (SourceModelIndex == -1 && RenderDataIndex == -1 && OtherIndex == -1); }
 	};
-	FString DefaultLODName = FString("Default");
-	TMap<FString, FLODName> ActiveLODNames;
-	void UpdateLODNames();
 
 	struct FLODMeshInfo
 	{
-		UE::Geometry::FDynamicMesh3 Mesh;
-
-		bool bInfoCached = false;
+		TSharedPtr<UE::Geometry::FDynamicMesh3> Mesh;
 		TArray<int> BoundaryEdges;
 	};
 
-	TMap<FString, TUniquePtr<FLODMeshInfo>> LODMeshCache;
+
+protected:
+	UStaticMesh* GetSingleStaticMesh();
+
+	ELODManagerToolActions PendingAction = ELODManagerToolActions::NoAction;
+	
+
+	bool bLODInfoValid = false;
+	// captures the material list and triangle and vertex counts for the current configuration of lods
+	void UpdateLODInfo();
+
+	// maps pretty name in UI to description of the LOD
+	TMap<FString, FLODName> ActiveLODNames;
+	void UpdateLODNames();
 
 	UPROPERTY()
 	TObjectPtr<UPreviewMesh> LODPreview;
@@ -254,7 +273,20 @@ protected:
 	void UpdatePreviewLOD();
 	void UpdatePreviewLines(FLODMeshInfo& LODMeshInfo);
 	void ClearPreviewLines();
-	bool CacheLODMesh(const FString& Name, FLODName LODName);
+
+	// returns the requested LOD as a dynamic mesh along with boundary edges.  Has internal cache
+	TUniquePtr<ULODManagerTool::FLODMeshInfo> GetLODMeshInfo(const FLODName& LODName);
+
+	// For undo / redo system with our custom changes.  
+	TUniquePtr< UE::Geometry::LODManagerHelper::FLODManagerToolChange > ActiveChange;
+	void BeginChange(FText TransactionName);
+	void EndChange();
+
+	// state information used for undo within the tool
+	TUniquePtr< UE::Geometry::LODManagerHelper::FProxyLODState>  ProxyLODState;
+
+	// cache of dynamic mesh representation for each lod, imported and renderdata.
+	TUniquePtr< UE::Geometry::LODManagerHelper::FDynamicMeshLODCache > DynamicMeshCache;
 };
 
 #if UE_ENABLE_INCLUDE_ORDER_DEPRECATED_IN_5_2
