ModelingTools: Remove single material optimization warning in AutoLOD tool.

diff --git a/Engine/Plugins/Editor/MeshLODToolset/Source/MeshLODToolset/Private/Graphs/GenerateMeshLODGraph.cpp b/Engine/Plugins/Editor/MeshLODToolset/Source/MeshLODToolset/Private/Graphs/GenerateMeshLODGraph.cpp
--- a/Engine/Plugins/Editor/MeshLODToolset/Source/MeshLODToolset/Private/Graphs/GenerateMeshLODGraph.cpp
+++ b/Engine/Plugins/Editor/MeshLODToolset/Source/MeshLODToolset/Private/Graphs/GenerateMeshLODGraph.cpp
@@ -365,10 +365,6 @@ void FGenerateMeshLODGraph::BuildGraph(const FDynamicMesh3* SourceMeshHint)
 		Graph->InferConnection(FinalMeshOutputNode, MatIDTransferNode);
 		SimplifyPhaseInputNode = MatIDTransferNode;
 	}
-	else
-	{
-		UE_LOG(LogGeometry, Warning, TEXT("AutoLOD: applying single-material optimizations"));
-	}
 
 	// need to compute valid normals before Simplify
 	//FGraph::FHandle PerVertexNormalsNode = Graph->AddNodeOfType<FComputeMeshPerVertexOverlayNormalsNode>(TEXT("PerVertexOverlayNormals"));
